MLDeformer: Add ChaosClothGenerator

diff --git a/Engine/Plugins/Animation/MLDeformer/ChaosClothGenerator/ChaosClothGenerator.uplugin b/Engine/Plugins/Animation/MLDeformer/ChaosClothGenerator/ChaosClothGenerator.uplugin
new file mode 100644
--- /dev/null
+++ b/Engine/Plugins/Animation/MLDeformer/ChaosClothGenerator/ChaosClothGenerator.uplugin
@@ -0,0 +1,38 @@
+{
+	"FileVersion": 3,
+	"Version": 0,
+	"VersionName": "0.0",
+	"FriendlyName": "Chaos Cloth Generator",
+	"Description": "Chaos Cloth Data Generator for ML Deformer",
+	"Category": "Animation",
+	"CreatedBy": "Epic Games, Inc.",
+	"CreatedByURL": "https://epicgames.com",
+	"DocsURL": "",
+	"MarketplaceURL": "",
+	"SupportURL": "https://forums.unrealengine.com/",
+	"CanContainContent": true,
+	"IsBetaVersion": false,
+	"IsExperimentalVersion": true,
+	"Installed": false,
+	"Modules": [
+		{
+			"Name": "ChaosClothGenerator",
+			"Type": "Editor",
+			"LoadingPhase": "Default"
+		}
+	],
+	"Plugins": [
+		{
+			"Name": "ChaosClothAsset",
+			"Enabled": true
+		},
+		{
+			"Name": "GeometryCache",
+			"Enabled": true
+		},
+		{
+			"Name": "MLDeformerFramework",
+			"Enabled": true
+		}
+	]
+}
\ No newline at end of file
diff --git a/Engine/Plugins/Animation/MLDeformer/ChaosClothGenerator/Source/ChaosClothGenerator/ChaosClothGenerator.build.cs b/Engine/Plugins/Animation/MLDeformer/ChaosClothGenerator/Source/ChaosClothGenerator/ChaosClothGenerator.build.cs
new file mode 100644
--- /dev/null
+++ b/Engine/Plugins/Animation/MLDeformer/ChaosClothGenerator/Source/ChaosClothGenerator/ChaosClothGenerator.build.cs
@@ -0,0 +1,59 @@
+// Copyright Epic Games, Inc. All Rights Reserved.
+
+using UnrealBuildTool;
+
+public class ChaosClothGenerator : ModuleRules
+{
+	public ChaosClothGenerator(ReadOnlyTargetRules Target) : base(Target)
+	{
+		PCHUsage = ModuleRules.PCHUsageMode.UseExplicitOrSharedPCHs;
+		
+		PublicIncludePaths.AddRange(
+			new string[] {
+				// ... add public include paths required here ...
+			}
+			);
+				
+		
+		PrivateIncludePaths.AddRange(
+			new string[] {
+				// ... add other private include paths required here ...
+			}
+			);
+			
+		
+		PublicDependencyModuleNames.AddRange(
+			new string[]
+			{
+				"Core"
+				// ... add other public dependencies that you statically link with here ...
+			}
+			);
+			
+		
+		PrivateDependencyModuleNames.AddRange(
+			new string[]
+			{
+				"CoreUObject",
+				"ChaosClothAssetEngine",
+				"Engine",
+				"GeometryCache",
+				"MLDeformerFramework",
+				"MLDeformerFrameworkEditor",
+				"PropertyEditor",
+				"RenderCore",
+				"Slate",
+				"SlateCore",
+				"UnrealEd"
+			}
+			);
+		
+		
+		DynamicallyLoadedModuleNames.AddRange(
+			new string[]
+			{
+				// ... add any modules that your module loads dynamically here ...
+			}
+			);
+	}
+}
diff --git a/Engine/Plugins/Animation/MLDeformer/ChaosClothGenerator/Source/ChaosClothGenerator/Private/ChaosClothGenerator.cpp b/Engine/Plugins/Animation/MLDeformer/ChaosClothGenerator/Source/ChaosClothGenerator/Private/ChaosClothGenerator.cpp
new file mode 100644
--- /dev/null
+++ b/Engine/Plugins/Animation/MLDeformer/ChaosClothGenerator/Source/ChaosClothGenerator/Private/ChaosClothGenerator.cpp
@@ -0,0 +1,812 @@
+// Copyright Epic Games, Inc. All Rights Reserved.
+
+#include "ChaosClothGenerator.h"
+
+#include "Animation/AnimSequence.h"
+#include "Animation/AttributesRuntime.h"
+#include "AssetRegistry/AssetRegistryModule.h"
+#include "BonePose.h"
+#include "ChaosClothAsset/ClothAsset.h"
+#include "ChaosClothAsset/ClothSimulationProxy.h"
+#include "ClothGeneratorComponent.h"
+#include "ClothGeneratorProperties.h"
+#include "Engine/SkinnedAssetCommon.h"
+#include "FileHelpers.h"
+#include "GeometryCache.h"
+#include "GeometryCacheCodecV1.h"
+#include "GeometryCacheMeshData.h"
+#include "GeometryCacheTrackStreamable.h"
+#include "GeometryCacheConstantTopologyWriter.h"
+#include "IDocumentation.h"
+#include "Internationalization/Regex.h"
+#include "Misc/AsyncTaskNotification.h"
+#include "Misc/Optional.h"
+#include "PhysicsEngine/PhysicsAsset.h"
+#include "Rendering/SkeletalMeshRenderData.h"
+#include "Rendering/SkeletalMeshModel.h"
+#include "SkeletalRenderPublic.h"
+#include "Tasks/Pipe.h"
+#include "UObject/SavePackage.h"
+
+DEFINE_LOG_CATEGORY(LogChaosClothGenerator);
+
+#define LOCTEXT_NAMESPACE "ChaosClothGenerator"
+
+namespace UE::Chaos::ClothGenerator
+{
+	namespace Private
+	{
+		TArray<int32> ParseFrames(const FString& FramesString)
+		{
+			TArray<int32> Result;
+			static const FRegexPattern AllowedCharsPattern(TEXT("^[-,0-9\\s]+$"));
+		
+			if (!FRegexMatcher(AllowedCharsPattern, FramesString).FindNext())
+			{
+			    UE_LOG(LogChaosClothGenerator, Error, TEXT("Input contains invalid characters."));
+			    return Result;
+			}
+		
+			static const FRegexPattern SingleNumberPattern(TEXT("^\\s*(\\d+)\\s*$"));
+			static const FRegexPattern RangePattern(TEXT("^\\s*(\\d+)\\s*-\\s*(\\d+)\\s*$"));
+		
+			TArray<FString> Segments;
+		    FramesString.ParseIntoArray(Segments, TEXT(","), true);
+		    for (const FString& Segment : Segments)
+		    {
+		    	bool bSegmentValid = false;
+		
+		    	FRegexMatcher SingleNumberMatcher(SingleNumberPattern, Segment);
+		    	if (SingleNumberMatcher.FindNext())
+		    	{
+		    	    const int32 SingleNumber = FCString::Atoi(*SingleNumberMatcher.GetCaptureGroup(1));
+		    	    Result.Add(SingleNumber);
+		    	    bSegmentValid = true;
+		    	}
+		    	else
+		    	{
+		    		FRegexMatcher RangeMatcher(RangePattern, Segment);
+		    		if (RangeMatcher.FindNext())
+		    		{
+		    		    const int32 RangeStart = FCString::Atoi(*RangeMatcher.GetCaptureGroup(1));
+		    		    const int32 RangeEnd = FCString::Atoi(*RangeMatcher.GetCaptureGroup(2));
+		
+		    		    for (int32 i = RangeStart; i <= RangeEnd; ++i)
+		    		    {
+		    		        Result.Add(i);
+		    		    }
+		    		    bSegmentValid = true;
+		    		}
+		    	}
+		    	
+		    	if (!bSegmentValid)
+		    	{
+		    	    UE_LOG(LogChaosClothGenerator, Error, TEXT("Invalid format in segment: %s"), *Segment);
+		    	}
+		    }
+		
+			return Result;
+		}
+		
+		TArray<int32> Range(int32 End)
+		{
+			TArray<int32> Result;
+			Result.Reserve(End);
+			for (int32 Index = 0; Index < End; ++Index)
+			{
+				Result.Add(Index);
+			}
+			return Result;
+		}
+		
+		TArray<uint32> Range(uint32 Start, uint32 End)
+		{
+			TArray<uint32> Result;
+			const uint32 Num = End - Start;
+			Result.SetNumUninitialized(Num);
+			for (uint32 Index = 0; Index < Num; ++Index)
+			{
+				Result.Add(Index + Start);
+			}
+			return Result;
+		}
+		
+		int32 GetNumVertices(const FSkeletalMeshLODRenderData& LODData)
+		{
+			int32 NumVertices = 0;
+			for(const FSkelMeshRenderSection& Section : LODData.RenderSections)
+			{
+				NumVertices += Section.NumVertices;
+			}
+			return NumVertices;
+		}
+		
+		TArrayView<TArray<FVector3f>> ShrinkToValidFrames(const TArrayView<TArray<FVector3f>>& Positions, int32 NumVertices)
+		{
+			int32 NumValidFrames = 0;
+			for (const TArray<FVector3f>& Frame : Positions)
+			{
+				if (Frame.Num() != NumVertices)
+				{
+					break;
+				}
+				++NumValidFrames;
+			}
+			return TArrayView<TArray<FVector3f>>(Positions.GetData(), NumValidFrames);
+		}
+		
+		void SaveGeometryCache(UGeometryCache& GeometryCache, const USkinnedAsset& Asset, TConstArrayView<uint32> ImportedVertexNumbers, TArrayView<TArray<FVector3f>> PositionsToMoveFrom)
+		{
+			const FSkeletalMeshRenderData* RenderData = Asset.GetResourceForRendering();
+			constexpr int32 LODIndex = 0;
+			if (!RenderData || !RenderData->LODRenderData.IsValidIndex(LODIndex))
+			{
+				return;
+			}
+			const FSkeletalMeshLODRenderData& LODData = RenderData->LODRenderData[LODIndex];
+			const int32 NumVertices = GetNumVertices(LODData);
+			PositionsToMoveFrom = ShrinkToValidFrames(PositionsToMoveFrom, NumVertices);
+		
+			using UE::GeometryCacheHelpers::FGeometryCacheConstantTopologyWriter;
+			using UE::GeometryCacheHelpers::AddTrackWriterFromSkinnedAsset;
+			using FTrackWriter = FGeometryCacheConstantTopologyWriter::FTrackWriter;
+			FGeometryCacheConstantTopologyWriter Writer(GeometryCache);
+			const int32 Index = AddTrackWriterFromSkinnedAsset(Writer, Asset);
+			if (Index == INDEX_NONE)
+			{
+				return;
+			}
+			FTrackWriter& TrackWriter = Writer.GetTrackWriter(Index);
+			TrackWriter.ImportedVertexNumbers = ImportedVertexNumbers;
+			TrackWriter.WriteAndClose(PositionsToMoveFrom);
+		}
+		
+		class FTimeScope
+		{
+		public:
+			explicit FTimeScope(FString InName)
+				: Name(MoveTemp(InName))
+				, StartTime(FDateTime::UtcNow())
+			{
+			}
+			~FTimeScope()
+			{
+				const FTimespan Duration = FDateTime::UtcNow() - StartTime;
+				UE_LOG(LogChaosClothGenerator, Log, TEXT("%s took %f secs"), *Name, Duration.GetTotalSeconds());
+			}
+		private:
+			FString Name;
+			FDateTime StartTime;
+		};
+		
+		void SavePackage(UObject& Object)
+		{
+			TArray<UPackage*> PackagesToSave = { Object.GetOutermost() };
+			constexpr bool bCheckDirty = false;
+			constexpr bool bPromptToSave = false;
+			FEditorFileUtils::PromptForCheckoutAndSave(PackagesToSave, bCheckDirty, bPromptToSave);
+		}
+		
+		TOptional<TArray<int32>> GetMeshImportVertexMap(const USkinnedAsset& SkeletalMeshAsset, const UChaosClothAsset& ClothAsset)
+		{
+			constexpr int32 LODIndex = 0;
+			const TOptional<TArray<int32>> None;
+			const FSkeletalMeshModel* const MLDModel = SkeletalMeshAsset.GetImportedModel();
+			if (!MLDModel || !MLDModel->LODModels.IsValidIndex(LODIndex))
+			{
+				return None;
+			}
+			const FSkeletalMeshLODModel& MLDLOD = MLDModel->LODModels[LODIndex];
+			const TArray<int32>& Map = MLDLOD.MeshToImportVertexMap;
+			if (Map.IsEmpty())
+			{
+				UE_LOG(LogChaosClothGenerator, Warning, TEXT("MeshToImportVertexMap is empty. MLDeformer Asset should be an imported SkeletalMesh (e.g. from fbx)."));
+				return None;
+			}
+			const FSkeletalMeshModel* const ClothModel = ClothAsset.GetImportedModel();
+			if (!ClothModel || !ClothModel->LODModels.IsValidIndex(LODIndex))
+			{
+				UE_LOG(LogChaosClothGenerator, Warning, TEXT("ClothAsset has no imported model."));
+				return None;
+			}
+			const FSkeletalMeshLODModel& ClothLOD = ClothModel->LODModels[LODIndex];
+		
+			if (MLDLOD.NumVertices != ClothLOD.NumVertices || MLDLOD.Sections.Num() != ClothLOD.Sections.Num())
+			{
+				UE_LOG(LogChaosClothGenerator, Warning, TEXT("SkeletalMeshAsset and ClothAsset have different number of vertices or sections. Check if the assets have the same mesh."));
+				return None;
+			}
+			
+			for (int32 SectionIndex = 0; SectionIndex < MLDLOD.Sections.Num(); ++SectionIndex)
+			{
+				const FSkelMeshSection& MLDSection = MLDLOD.Sections[SectionIndex];
+				const FSkelMeshSection& ClothSection = ClothLOD.Sections[SectionIndex];
+				if (MLDSection.NumVertices != ClothSection.NumVertices)
+				{
+					UE_LOG(LogChaosClothGenerator, Warning, TEXT("SkeletalMeshAsset and ClothAsset have different number of vertices in section %d. Check if the assets have the same mesh."), SectionIndex);
+					return None;
+				}
+				for (int32 VertexIndex = 0; VertexIndex < MLDSection.NumVertices; ++VertexIndex)
+				{
+					const FVector3f& MLDPosition = MLDSection.SoftVertices[VertexIndex].Position;
+					const FVector3f& ClothPosition = ClothSection.SoftVertices[VertexIndex].Position;
+					if (!MLDPosition.Equals(ClothPosition, UE_KINDA_SMALL_NUMBER))
+					{
+						UE_LOG(LogChaosClothGenerator, Warning, TEXT("SkeletalMeshAsset and ClothAsset have different vertex positions. Check if the assets have the same vertex order."));
+						return None;
+					}
+				}
+			}
+		
+			return Map;
+		}
+	};
+
+	void FChaosClothGenerator::Tick(float DeltaTime)
+	{
+		if (PendingAction == EClothGeneratorActions::StartGenerate)
+		{
+			StartGenerate();
+		}
+		else if (PendingAction == EClothGeneratorActions::TickGenerate)
+		{
+			TickGenerate();
+		}
+	}
+
+	TStatId FChaosClothGenerator::GetStatId() const
+	{
+		RETURN_QUICK_DECLARE_CYCLE_STAT(FChaosClothGenerator, STATGROUP_Tickables);
+	}
+
+	template<typename TaskType>
+	class FChaosClothGenerator::TTaskRunner : public FNonAbandonableTask
+	{
+	public:
+		TTaskRunner(TUniquePtr<TaskType> InTask)
+			: Task(MoveTemp(InTask))
+		{
+		}
+	
+		void DoWork()
+		{
+			if (Task)
+			{
+				Task->DoWork();
+			}
+		}
+	
+		FORCEINLINE TStatId GetStatId() const
+		{
+			RETURN_QUICK_DECLARE_CYCLE_STAT(TTaskRunner, STATGROUP_ThreadPoolAsyncTasks);
+		}
+	
+	private:
+		TUniquePtr<TaskType> Task;
+	};
+	
+	struct FChaosClothGenerator::FSimResource
+	{
+		TStrongObjectPtr<UClothGeneratorComponent> ClothComponent = nullptr;
+		TSharedPtr<FProxy> Proxy;
+		TUniquePtr<UE::Tasks::FPipe> Pipe;
+		FEvent* SkinEvent = nullptr;
+		std::atomic<bool> bNeedsSkin = false;
+		TArrayView<TArray<FVector3f>> SimulatedPositions;
+		std::atomic<int32>* NumSimulatedFrames = nullptr;
+		std::atomic<bool>* bCancelled = nullptr;
+
+		bool IsCancelled() const { return !bCancelled || bCancelled->load(); }
+		void FinishFrame()
+		{
+			if (NumSimulatedFrames)
+			{
+				++(*NumSimulatedFrames); 
+			}
+		}
+	};
+	
+	struct FChaosClothGenerator::FTaskResource
+	{
+		TArray<FSimResource> SimResources;
+
+		TUniquePtr<FExecuterType> Executer;
+		TUniquePtr<FAsyncTaskNotification> Notification;
+		FDateTime StartTime;
+		FDateTime LastUpdateTime;
+
+		TArray<int32> FramesToSimulate;
+		TArray<TArray<FVector3f>> SimulatedPositions;
+		TArray<uint32> ImportedVertexNumbers;
+		UGeometryCache* Cache = nullptr;
+
+		std::atomic<int32> NumSimulatedFrames = 0;
+		std::atomic<bool> bCancelled = false;
+
+		UWorld* World = nullptr;
+	
+		bool AllocateSimResources_GameThread(UChaosClothAsset& Asset, int32 Num);
+		void FreeSimResources_GameThread();
+		void FlushRendering();
+		void Cancel();
+	};
+
+	class FChaosClothGenerator::FLaunchSimsTask
+	{
+	public: 
+		FLaunchSimsTask(FTaskResource& InTaskResource, TStrongObjectPtr<UClothGeneratorProperties> InProperties)
+			: TaskResource(InTaskResource)
+			, SimResources(InTaskResource.SimResources)
+			, Properties(InProperties)
+		{}
+	
+		void DoWork();
+	
+	private:
+		using FPipe = UE::Tasks::FPipe;
+	
+		enum class ESaveType
+		{
+			LastStep,
+			EveryStep,
+		};
+	
+		void Simulate(FSimResource &SimResource, int32 AnimFrame, int32 CacheFrame) const;
+		void PrepareAnimationSequence();
+		void RestoreAnimationSequence();
+		TArray<FTransform> GetBoneTransforms(UChaosClothComponent& InClothComponent, int32 Frame) const;
+		TArray<FVector3f> GetRenderPositions(FSimResource& SimResource) const;
+	
+		FTaskResource& TaskResource;
+		TArray<FSimResource>& SimResources;
+		TStrongObjectPtr<UClothGeneratorProperties> Properties;
+		EAnimInterpolationType InterpolationTypeBackup = EAnimInterpolationType::Linear;
+	};
+
+	bool FChaosClothGenerator::FTaskResource::AllocateSimResources_GameThread(UChaosClothAsset& Asset, int32 Num)
+	{
+		World = UWorld::CreateWorld(EWorldType::None, false);
+		SimResources.SetNum(Num);
+		for(int32 Index = 0; Index < Num; ++Index)
+		{
+			UClothGeneratorComponent* const CopyComponent = NewObject<UClothGeneratorComponent>();
+			CopyComponent->SetClothAsset(&Asset);
+			CopyComponent->RegisterComponentWithWorld(World);
+	
+			USkinnedMeshComponent* const PoseComponent = CopyComponent->LeaderPoseComponent.Get() ? CopyComponent->LeaderPoseComponent.Get() : CopyComponent;
+			constexpr int32 LODIndex = 0;
+			PoseComponent->SetForcedLOD(LODIndex + 1);
+			PoseComponent->UpdateLODStatus();
+			PoseComponent->RefreshBoneTransforms(nullptr);
+			CopyComponent->bRenderStatic = false;
+			constexpr bool bRecreateRenderStateImmediately = true;
+			CopyComponent->SetCPUSkinningEnabled(true, bRecreateRenderStateImmediately);
+			CopyComponent->ResumeSimulation();
+	
+			FSimResource& SimResource = SimResources[Index];
+			SimResource.ClothComponent = TStrongObjectPtr<UClothGeneratorComponent>(CopyComponent);
+			SimResource.Proxy = CopyComponent->GetProxy().Pin();
+			check(SimResource.Proxy != nullptr);
+			SimResource.Pipe = MakeUnique<UE::Tasks::FPipe>(*FString::Printf(TEXT("SimPipe:%d"), Index));
+			SimResource.SkinEvent = FPlatformProcess::GetSynchEventFromPool();
+			SimResource.bNeedsSkin.store(false);
+	
+			SimResource.SimulatedPositions = TArrayView<TArray<FVector3f>>(SimulatedPositions);
+			SimResource.NumSimulatedFrames = &NumSimulatedFrames;
+			SimResource.bCancelled = &bCancelled;
+	
+			if (CopyComponent == nullptr || SimResource.Proxy == nullptr || SimResource.Pipe == nullptr)
+			{
+				UE_LOG(LogChaosClothGenerator, Error, TEXT("Failed to allocate simulation resources"));
+				return false;
+			}
+		}
+		return true;
+	}
+	
+	void FChaosClothGenerator::FTaskResource::FreeSimResources_GameThread()
+	{
+		if (Executer.IsValid())
+		{
+			Executer->EnsureCompletion();
+		}
+		for (FSimResource& SimResource : SimResources)
+		{
+			FPlatformProcess::ReturnSynchEventToPool(SimResource.SkinEvent);
+			SimResource.Pipe.Reset();
+			SimResource.ClothComponent->UnregisterComponent();
+			SimResource.ClothComponent->DestroyComponent();
+		}
+		SimResources.Empty();
+		World->DestroyWorld(false);
+	}
+	
+	void FChaosClothGenerator::FTaskResource::FlushRendering()
+	{
+		// Copy bNeedsSkin
+		TArray<bool> NeedsSkin;
+		NeedsSkin.SetNum(SimResources.Num());
+		bool bAnyNeedsSkin = false;
+		for (int32 Index = 0; Index < SimResources.Num(); ++Index)
+		{
+			const bool bNeedsSkin = SimResources[Index].bNeedsSkin.load();
+			bAnyNeedsSkin |= bNeedsSkin;
+			NeedsSkin[Index] = bNeedsSkin;
+		}
+	
+		if (bAnyNeedsSkin)
+		{
+			FlushRenderingCommands();
+			for (int32 Index = 0; Index < SimResources.Num(); ++Index)
+			{
+				if (NeedsSkin[Index])
+				{
+					SimResources[Index].bNeedsSkin.store(false);
+					SimResources[Index].SkinEvent->Trigger();
+				}
+			}
+		}
+	}
+
+	void FChaosClothGenerator::FTaskResource::Cancel()
+	{
+		bCancelled.store(true);
+	}
+
+	void FChaosClothGenerator::FLaunchSimsTask::DoWork()
+	{	
+		const int32 NumFrames = TaskResource.FramesToSimulate.Num();
+		PrepareAnimationSequence();
+	
+		const int32 NumThreads = Properties->bDebug ? 1 : Properties->NumThreads;
+	
+		for (int32 Frame = 0; Frame < NumFrames; Frame++)
+		{
+			if (!TaskResource.bCancelled.load())
+			{
+				const int32 ThreadIdx = Frame % NumThreads;
+				const int32 AnimFrame = TaskResource.FramesToSimulate[Frame];
+	
+				FSimResource& SimResource = SimResources[ThreadIdx];
+				SimResource.Pipe->Launch(*FString::Printf(TEXT("SimFrame:%d"), AnimFrame), [this, &SimResource, AnimFrame, Frame]()
+				{ 
+					FMemMark Mark(FMemStack::Get());
+					Simulate(SimResource, AnimFrame, Frame);
+				});
+			}
+			else
+			{
+				break;
+			}
+		}
+	
+		for (FSimResource& SimResource : SimResources)
+		{
+			SimResource.Pipe->WaitUntilEmpty();
+		}
+	
+		RestoreAnimationSequence();
+	}
+
+	void FChaosClothGenerator::FLaunchSimsTask::Simulate(FSimResource &SimResource, int32 AnimFrame, int32 CacheFrame) const
+	{
+		UClothGeneratorComponent& TaskComponent = *SimResource.ClothComponent;
+		FProxy& DataGenerationProxy = *SimResource.Proxy;
+	
+		const float TimeStep = Properties->TimeStep;
+		const int32 NumSteps = Properties->NumSteps;
+		const ESaveType SaveType = Properties->bDebug ? ESaveType::EveryStep : ESaveType::LastStep;
+	
+		const TArray<FTransform> Transforms = GetBoneTransforms(TaskComponent, AnimFrame);
+		TaskComponent.Pose(Transforms);
+		TaskComponent.ForceNextUpdateTeleportAndReset();
+		DataGenerationProxy.FillSimulationContext(TimeStep);
+		DataGenerationProxy.InitializeConfigs();
+		bool bCancelled = false;
+		for (int32 Step = 0; Step < NumSteps; ++Step)
+		{
+			if (SimResource.IsCancelled())
+			{
+				bCancelled = true;
+				break;
+			}
+			else 
+			{
+				DataGenerationProxy.Tick();
+	
+				// Clear any reset flags at the end of the first step
+				if (Step == 0 && NumSteps > 1)
+				{
+					TaskComponent.ResetTeleportMode();
+					DataGenerationProxy.FillSimulationContext(TimeStep);	
+				}
+	
+				if (SaveType == ESaveType::EveryStep)
+				{
+					DataGenerationProxy.WriteSimulationData();
+					SimResource.SimulatedPositions[Step] = GetRenderPositions(SimResource);
+				}
+			}
+		}
+	
+		if (SaveType == ESaveType::LastStep && !bCancelled)
+		{
+			DataGenerationProxy.WriteSimulationData();
+			SimResource.SimulatedPositions[CacheFrame] = GetRenderPositions(SimResource);
+		}
+
+		SimResource.FinishFrame();
+	}
+	
+	void FChaosClothGenerator::FLaunchSimsTask::PrepareAnimationSequence()
+	{
+		TObjectPtr<UAnimSequence> AnimationSequence = Properties->AnimationSequence;
+		if (AnimationSequence)
+		{
+			InterpolationTypeBackup = AnimationSequence->Interpolation;
+			AnimationSequence->Interpolation = EAnimInterpolationType::Step;
+		}
+	}
+	
+	void FChaosClothGenerator::FLaunchSimsTask::RestoreAnimationSequence()
+	{
+		TObjectPtr<UAnimSequence> AnimationSequence = Properties->AnimationSequence;
+		if (AnimationSequence)
+		{
+			AnimationSequence->Interpolation = InterpolationTypeBackup;
+		}
+	}
+	
+	TArray<FTransform> FChaosClothGenerator::FLaunchSimsTask::GetBoneTransforms(UChaosClothComponent& InClothComponent, int32 Frame) const
+	{
+		const UAnimSequence* AnimationSequence = Properties->AnimationSequence;
+		const double Time = FMath::Clamp(AnimationSequence->GetSamplingFrameRate().AsSeconds(Frame), 0., (double)AnimationSequence->GetPlayLength());
+		FAnimExtractContext ExtractionContext(Time);
+	
+		UChaosClothAsset* const ClothAsset = InClothComponent.GetClothAsset();
+		const FReferenceSkeleton* const ReferenceSkeleton = ClothAsset ? &ClothAsset->GetRefSkeleton() : nullptr;
+		USkeleton* const Skeleton = ClothAsset ? ClothAsset->GetSkeleton() : nullptr;
+		const int32 NumBones = ReferenceSkeleton ? ReferenceSkeleton->GetNum() : 0;
+	
+		TArray<uint16> BoneIndices;
+		BoneIndices.SetNumUninitialized(NumBones);
+		for (int32 Index = 0; Index < NumBones; ++Index)
+		{
+			BoneIndices[Index] = (uint16)Index;
+		}
+	
+		FBoneContainer BoneContainer;
+		BoneContainer.SetUseRAWData(true);
+		BoneContainer.InitializeTo(BoneIndices, UE::Anim::FCurveFilterSettings(), *Skeleton);
+	
+		FCompactPose OutPose;
+		OutPose.SetBoneContainer(&BoneContainer);
+		FBlendedCurve OutCurve;
+		OutCurve.InitFrom(BoneContainer);
+		UE::Anim::FStackAttributeContainer TempAttributes;
+	
+		FAnimationPoseData AnimationPoseData(OutPose, OutCurve, TempAttributes);
+		AnimationSequence->GetAnimationPose(AnimationPoseData, ExtractionContext);
+	
+		TArray<FTransform> ComponentSpaceTransforms;
+		ComponentSpaceTransforms.SetNumUninitialized(NumBones);
+		for (int32 Index = 0; Index < NumBones; ++Index)
+		{
+			const FCompactPoseBoneIndex CompactIndex = BoneContainer.MakeCompactPoseIndex(FMeshPoseBoneIndex(Index));
+			const int32 ParentIndex = ReferenceSkeleton->GetParentIndex(Index);
+			ComponentSpaceTransforms[Index] = 
+				ComponentSpaceTransforms.IsValidIndex(ParentIndex) && ParentIndex < Index ? 
+				AnimationPoseData.GetPose()[CompactIndex] * ComponentSpaceTransforms[ParentIndex] : 
+				ReferenceSkeleton->GetRefBonePose()[Index];
+		}
+	
+		return ComponentSpaceTransforms;
+	}
+	
+	TArray<FVector3f> FChaosClothGenerator::FLaunchSimsTask::GetRenderPositions(FSimResource& SimResource) const
+	{
+		check(SimResource.ClothComponent);
+		TArray<FVector3f> Positions;
+		TArray<FFinalSkinVertex> OutVertices;
+		// This could potentially be slow. 
+		SimResource.ClothComponent->RecreateRenderState_Concurrent();
+		SimResource.bNeedsSkin.store(true);
+		SimResource.SkinEvent->Wait();
+		
+		SimResource.ClothComponent->GetCPUSkinnedCachedFinalVertices(OutVertices);
+		Positions.SetNum(OutVertices.Num());
+		for (int32 Index = 0; Index < OutVertices.Num(); ++Index)
+		{
+			Positions[Index] = OutVertices[Index].Position;
+		}
+		return Positions;
+	}
+	
+	FChaosClothGenerator::FChaosClothGenerator()
+	{
+		Properties = TStrongObjectPtr<UClothGeneratorProperties>(NewObject<UClothGeneratorProperties>());
+	}
+
+	FChaosClothGenerator::~FChaosClothGenerator()
+	{
+		if (TaskResource != nullptr)
+		{
+			TaskResource->FreeSimResources_GameThread();
+		}
+	}
+
+	void FChaosClothGenerator::StartGenerate()
+	{
+		check(PendingAction == EClothGeneratorActions::StartGenerate);
+		if (Properties->ClothAsset == nullptr)
+		{
+			UE_LOG(LogChaosClothGenerator, Error, TEXT("ClothAsset is null."));
+			PendingAction = EClothGeneratorActions::NoAction;
+			return;
+		}
+		if (Properties->SkeletalMeshAsset == nullptr)
+		{
+			UE_LOG(LogChaosClothGenerator, Error, TEXT("SkeletalMeshAsset is null."));
+			PendingAction = EClothGeneratorActions::NoAction;
+			return;
+		}
+		if (Properties->AnimationSequence == nullptr)
+		{
+			UE_LOG(LogChaosClothGenerator, Error, TEXT("AnimationSequence is null."));
+			PendingAction = EClothGeneratorActions::NoAction;
+			return;
+		}
+		UGeometryCache* const Cache = GetCache();
+		if (Cache == nullptr)
+		{
+			UE_LOG(LogChaosClothGenerator, Error, TEXT("Cannot find or create geometry cache."));
+			PendingAction = EClothGeneratorActions::NoAction;
+			return;
+		}
+		if (TaskResource != nullptr)
+		{
+			UE_LOG(LogChaosClothGenerator, Error, TEXT("Previous generation is still running."));
+			PendingAction = EClothGeneratorActions::NoAction;
+			return;
+		}
+
+		using UE::Chaos::ClothGenerator::Private::GetMeshImportVertexMap;
+		TOptional<TArray<int32>> OptionalMap = GetMeshImportVertexMap(*Properties->SkeletalMeshAsset, *Properties->ClothAsset);
+		if (!OptionalMap)
+		{
+			PendingAction = EClothGeneratorActions::NoAction;
+			return;
+		}
+		TaskResource = MakeUnique<FTaskResource>();
+
+		using Private::ParseFrames;
+		using Private::Range;
+		TaskResource->FramesToSimulate = Properties->bDebug 
+			? TArray<int32>{ (int32)Properties->DebugFrame } 
+			: (Properties->FramesToSimulate.Len() > 0
+				? ParseFrames(Properties->FramesToSimulate) 
+				: Range(Properties->AnimationSequence->GetNumberOfSampledKeys()));
+		const int32 NumFrames = TaskResource->FramesToSimulate.Num();
+		if (NumFrames == 0)
+		{
+			PendingAction = EClothGeneratorActions::NoAction;
+			return;
+		}
+		TaskResource->SimulatedPositions.SetNum(Properties->bDebug ? Properties->NumSteps : NumFrames);
+
+		if (!TaskResource->AllocateSimResources_GameThread(*Properties->ClothAsset, Properties->NumThreads))
+		{
+			PendingAction = EClothGeneratorActions::NoAction;
+			return;
+		}
+		TaskResource->Cache = Cache;
+
+		TUniquePtr<FLaunchSimsTask> Task = MakeUnique<FLaunchSimsTask>(*TaskResource, Properties);
+		TaskResource->Executer = MakeUnique<FExecuterType>(MoveTemp(Task));
+		TaskResource->Executer->StartBackgroundTask();
+	
+		FAsyncTaskNotificationConfig NotificationConfig;
+		NotificationConfig.TitleText = LOCTEXT("SimulateCloth", "Simulating Cloth");
+		NotificationConfig.ProgressText = FText::FromString(TEXT("0%"));
+		NotificationConfig.bCanCancel = true;
+		NotificationConfig.bKeepOpenOnSuccess = true;
+		NotificationConfig.bKeepOpenOnFailure = true;
+		TaskResource->Notification = MakeUnique<FAsyncTaskNotification>(NotificationConfig);
+		TaskResource->StartTime = FDateTime::UtcNow();
+		TaskResource->LastUpdateTime = TaskResource->StartTime;
+	
+		const TArray<int32>& Map = OptionalMap.GetValue();
+		TaskResource->ImportedVertexNumbers = TArray<uint32>(reinterpret_cast<const uint32*>(Map.GetData()), Map.Num());
+	
+		PendingAction = EClothGeneratorActions::TickGenerate;
+	}
+	
+	void FChaosClothGenerator::TickGenerate()
+	{
+		check(PendingAction == EClothGeneratorActions::TickGenerate && TaskResource != nullptr);
+			
+		bool bFinished = false;
+		const bool bCancelled = TaskResource->Notification->GetPromptAction() == EAsyncTaskNotificationPromptAction::Cancel;
+		if (TaskResource->Executer->IsDone())
+		{
+			bFinished = true;
+		}
+		else if (bCancelled)
+		{
+			TaskResource->Cancel();
+			bFinished = true;
+		}
+			
+		if (!bFinished)
+		{
+			TaskResource->FlushRendering();
+			const FDateTime CurrentTime = FDateTime::UtcNow();
+			const double SinceLastUpdate = (CurrentTime - TaskResource->LastUpdateTime).GetTotalSeconds();
+			if (SinceLastUpdate < 0.2)
+			{
+				return;
+			}
+			
+			const int32 NumSimulatedFrames = TaskResource->NumSimulatedFrames.load();
+			const int32 NumTotalFrames = TaskResource->FramesToSimulate.Num();
+			const FText ProgressMessage = FText::FromString(FString::Printf(TEXT("Finished %d/%d, %.1f%%"), NumSimulatedFrames, NumTotalFrames, 100.0 * NumSimulatedFrames / NumTotalFrames));
+			TaskResource->Notification->SetProgressText(ProgressMessage);
+			TaskResource->LastUpdateTime = CurrentTime;
+		}
+		else
+		{
+			FreeTaskResource(bCancelled);
+			PendingAction = EClothGeneratorActions::NoAction;
+		}
+	}
+
+	UClothGeneratorProperties& FChaosClothGenerator::GetProperties() const
+	{
+		return *Properties;
+	}
+	
+	void FChaosClothGenerator::RequestAction(EClothGeneratorActions ActionType)
+	{
+		if (PendingAction != EClothGeneratorActions::NoAction)
+		{
+			return;
+		}
+		PendingAction = ActionType;
+	}
+	
+	
+	UGeometryCache* FChaosClothGenerator::GetCache() const
+	{
+		return Properties->bDebug ? Properties->DebugCache : Properties->SimulatedCache;
+	}
+	
+	void FChaosClothGenerator::FreeTaskResource(bool bCancelled)
+	{
+		TaskResource->Notification->SetProgressText(LOCTEXT("Finishing", "Finishing, please wait"));
+		TaskResource->FreeSimResources_GameThread();
+		const FDateTime CurrentTime = FDateTime::UtcNow();
+		UE_LOG(LogChaosClothGenerator, Log, TEXT("Training finished in %f seconds"), (CurrentTime - TaskResource->StartTime).GetTotalSeconds());
+	
+		{
+			UE::Chaos::ClothGenerator::Private::FTimeScope TimeScope(TEXT("Saving"));
+	
+			using UE::Chaos::ClothGenerator::Private::SaveGeometryCache;
+			using UE::Chaos::ClothGenerator::Private::SavePackage;
+			SaveGeometryCache(*TaskResource->Cache, *Properties->ClothAsset, TaskResource->ImportedVertexNumbers, TaskResource->SimulatedPositions);
+			SavePackage(*TaskResource->Cache);
+		}
+		if (bCancelled)
+		{
+			TaskResource->Notification->SetProgressText(LOCTEXT("Cancelled", "Cancelled"));
+			TaskResource->Notification->SetComplete(false);
+		}
+		else
+		{
+			TaskResource->Notification->SetProgressText(LOCTEXT("Finished", "Finished"));
+			TaskResource->Notification->SetComplete(true);
+		}
+		TaskResource.Reset();
+	}
+};
+
+#undef LOCTEXT_NAMESPACE
diff --git a/Engine/Plugins/Animation/MLDeformer/ChaosClothGenerator/Source/ChaosClothGenerator/Private/ChaosClothGenerator.h b/Engine/Plugins/Animation/MLDeformer/ChaosClothGenerator/Source/ChaosClothGenerator/Private/ChaosClothGenerator.h
new file mode 100644
--- /dev/null
+++ b/Engine/Plugins/Animation/MLDeformer/ChaosClothGenerator/Source/ChaosClothGenerator/Private/ChaosClothGenerator.h
@@ -0,0 +1,61 @@
+// Copyright Epic Games, Inc. All Rights Reserved.
+
+#pragma once
+
+#include "CoreTypes.h"
+#include "Async/AsyncWork.h"
+#include "TickableEditorObject.h"
+#include "UObject/StrongObjectPtr.h"
+
+DECLARE_LOG_CATEGORY_EXTERN(LogChaosClothGenerator, Log, All);
+
+class UGeometryCache;
+class UClothGeneratorProperties;
+namespace UE::Chaos::ClothGenerator
+{
+	class FClothGeneratorProxy;
+
+	enum class EClothGeneratorActions
+	{
+		NoAction,
+		StartGenerate,
+		TickGenerate
+	};
+
+	class FChaosClothGenerator : public FTickableEditorObject
+	{
+	public:
+		FChaosClothGenerator();
+		virtual ~FChaosClothGenerator();
+
+		//~ Begin FTickableEditorObject Interface
+		virtual void Tick(float DeltaTime) override;
+		virtual TStatId GetStatId() const override;
+		virtual ETickableTickType GetTickableTickType() const override { return ETickableTickType::Always; }
+		//~ End FTickableEditorObject Interface
+
+		UClothGeneratorProperties& GetProperties() const;
+		void RequestAction(EClothGeneratorActions Action);
+	private:
+		struct FSimResource;
+		struct FTaskResource;
+		class FLaunchSimsTask;
+		// Using dummy TTaskRunner to avoid adding declaration of FTaskResource and FLaunchSimsTask to header file
+		template<typename TaskType>
+		class TTaskRunner;
+
+		using FProxy = FClothGeneratorProxy;
+		using FExecuterType = FAsyncTask<TTaskRunner<FLaunchSimsTask>>;
+
+		void StartGenerate();
+		void TickGenerate();
+		void FreeTaskResource(bool bCancelled);
+		UGeometryCache* GetCache() const;
+
+		TStrongObjectPtr<UClothGeneratorProperties> Properties;
+		EClothGeneratorActions PendingAction = EClothGeneratorActions::NoAction;
+		TUniquePtr<FTaskResource> TaskResource;
+
+		
+	};
+};
\ No newline at end of file
diff --git a/Engine/Plugins/Animation/MLDeformer/ChaosClothGenerator/Source/ChaosClothGenerator/Private/ChaosClothGeneratorModule.cpp b/Engine/Plugins/Animation/MLDeformer/ChaosClothGenerator/Source/ChaosClothGenerator/Private/ChaosClothGeneratorModule.cpp
new file mode 100644
--- /dev/null
+++ b/Engine/Plugins/Animation/MLDeformer/ChaosClothGenerator/Source/ChaosClothGenerator/Private/ChaosClothGeneratorModule.cpp
@@ -0,0 +1,27 @@
+// Copyright Epic Games, Inc. All Rights Reserved.
+#include "ChaosClothGeneratorToolsMenuExtender.h"
+#include "MLDeformerEditorToolkit.h"
+#include "PropertyEditorModule.h"
+#include "SClothGeneratorWidget.h"
+
+namespace UE::Chaos::ClothGenerator
+{
+	class FChaosClothGeneratorModule : public IModuleInterface
+	{
+	public:
+		virtual void StartupModule() override
+		{
+			UE::MLDeformer::FMLDeformerEditorToolkit::AddToolsMenuExtender(CreateToolsMenuExtender());
+			FPropertyEditorModule& PropertyModule = FModuleManager::LoadModuleChecked<FPropertyEditorModule>("PropertyEditor");
+			PropertyModule.RegisterCustomClassLayout("ClothGeneratorProperties", FOnGetDetailCustomizationInstance::CreateStatic(&FClothGeneratorDetails::MakeInstance));
+		}
+
+		virtual void ShutdownModule() override
+		{
+			FPropertyEditorModule& PropertyModule = FModuleManager::LoadModuleChecked<FPropertyEditorModule>("PropertyEditor");
+			PropertyModule.UnregisterCustomClassLayout("ClothGeneratorProperties");
+		}
+	};
+};
+
+IMPLEMENT_MODULE(UE::Chaos::ClothGenerator::FChaosClothGeneratorModule, ChaosClothGenerator)
\ No newline at end of file
diff --git a/Engine/Plugins/Animation/MLDeformer/ChaosClothGenerator/Source/ChaosClothGenerator/Private/ChaosClothGeneratorToolsMenuExtender.cpp b/Engine/Plugins/Animation/MLDeformer/ChaosClothGenerator/Source/ChaosClothGenerator/Private/ChaosClothGeneratorToolsMenuExtender.cpp
new file mode 100644
--- /dev/null
+++ b/Engine/Plugins/Animation/MLDeformer/ChaosClothGenerator/Source/ChaosClothGenerator/Private/ChaosClothGeneratorToolsMenuExtender.cpp
@@ -0,0 +1,119 @@
+// Copyright Epic Games, Inc. All Rights Reserved.
+#include "ChaosClothGeneratorToolsMenuExtender.h"
+
+#include "ClothGeneratorProperties.h"
+#include "Engine/SkeletalMesh.h"
+#include "Framework/Docking/TabManager.h"
+#include "IDocumentation.h"
+#include "MLDeformerGeomCacheEditorModel.h"
+#include "MLDeformerEditorModel.h"
+#include "MLDeformerModel.h"
+#include "MLDeformerGeomCacheModel.h"
+#include "SClothGeneratorWidget.h"
+#include "Widgets/Docking/SDockTab.h"
+
+#define LOCTEXT_NAMESPACE "ChaosClothGeneratorToolsMenuExtender"
+
+namespace UE::Chaos::ClothGenerator
+{
+	namespace Private
+	{
+		void SpawnTab(UE::MLDeformer::FMLDeformerEditorToolkit& Toolkit)
+		{
+			TSharedPtr<SDockTab> Tab = Toolkit.GetAssociatedTabManager()->TryInvokeTab(FName(FChaosClothGeneratorTabSummoner::TabID));
+			if (!Tab.IsValid())
+			{
+				return;
+			}
+			
+			TSharedRef<SWidget> TabContent = Tab->GetContent();
+			if (TabContent == SNullWidget::NullWidget || TabContent->GetTypeAsString() != "SClothGeneratorWidget")
+			{
+				return;
+			}
+			TSharedRef<SClothGeneratorWidget> Widget = StaticCastSharedRef<SClothGeneratorWidget>(Tab->GetContent());
+			if (Widget == SNullWidget::NullWidget)
+			{
+				return;
+			}
+		
+			TWeakObjectPtr<UClothGeneratorProperties> Properties = Widget->GetProperties();
+			if (!Properties.IsValid())
+			{
+				return;
+			}
+		
+			using UE::MLDeformer::FMLDeformerEditorModel;
+			if (FMLDeformerEditorModel* EditorModel = Toolkit.GetActiveModel())
+			{
+				if (UMLDeformerModel* Model = EditorModel->GetModel())
+				{
+					Properties->SkeletalMeshAsset = Model->GetSkeletalMesh();
+					Properties->AnimationSequence = Model->GetAnimSequence();
+					if (UMLDeformerGeomCacheModel* GeomCacheModel = Cast<UMLDeformerGeomCacheModel>(Model))
+					{
+						if (UGeometryCache* GeomCache = GeomCacheModel->GetGeometryCache())
+						{
+							Properties->SimulatedCache = GeomCache;
+						}
+					}
+				}
+			}
+		}
+	};
+	
+	TUniquePtr<FChaosClothGeneratorToolsMenuExtender> CreateToolsMenuExtender()
+	{
+		return MakeUnique<FChaosClothGeneratorToolsMenuExtender>();
+	}
+	
+	const FName FChaosClothGeneratorTabSummoner::TabID = FName("ChaosClothGenerator");
+	
+	FChaosClothGeneratorTabSummoner::FChaosClothGeneratorTabSummoner(const TSharedRef<FMLDeformerEditorToolkit>& InEditor)
+		: FWorkflowTabFactory(TabID, InEditor)
+	{
+		bIsSingleton = true;
+		TabLabel = LOCTEXT("ChaosClothGenerator", "Chaos Cloth Generator");
+		ViewMenuDescription = LOCTEXT("ViewMenu_Desc", "Chaos Cloth Generator");
+		ViewMenuTooltip = LOCTEXT("ViewMenu_ToolTip", "Show the Chaos Cloth Generator tab.");
+	}
+	
+	TSharedRef<SWidget> FChaosClothGeneratorTabSummoner::CreateTabBody(const FWorkflowTabSpawnInfo& Info) const
+	{
+		return SNew(SClothGeneratorWidget);
+	}
+	
+	TSharedPtr<SToolTip> FChaosClothGeneratorTabSummoner::CreateTabToolTipWidget(const FWorkflowTabSpawnInfo& Info) const
+	{
+		return IDocumentation::Get()->CreateToolTip(
+			LOCTEXT("ChaosClothGeneratorTooltip", "Generate training data using chaos cloth solver."),
+			NULL,
+			TEXT("Shared/Editors/Persona"),
+			TEXT("ChaosClothGenerator_Window"));
+	}
+	
+	FMenuEntryParams FChaosClothGeneratorToolsMenuExtender::GetMenuEntry(FMLDeformerEditorToolkit& Toolkit) const
+	{
+		FMenuEntryParams Params;
+		Params.DirectActions = FUIAction(
+			FExecuteAction::CreateLambda([&Toolkit]()
+			{
+				Private::SpawnTab(Toolkit);
+			}),
+			FCanExecuteAction::CreateLambda([]()
+			{
+				return true;
+			})
+		);
+		Params.LabelOverride = LOCTEXT("ChaosClothGenerator", "Chaos Cloth Generator");
+		Params.ToolTipOverride = LOCTEXT("ChaosClothGeneratorTooltip", "Generate training data using chaos cloth solver");
+		return Params;
+	}
+	
+	TSharedPtr<FWorkflowTabFactory> FChaosClothGeneratorToolsMenuExtender::GetTabSummoner(const TSharedRef<FMLDeformerEditorToolkit>& Toolkit) const
+	{
+		return MakeShared<FChaosClothGeneratorTabSummoner>(Toolkit);
+	}
+};
+
+#undef LOCTEXT_NAMESPACE
\ No newline at end of file
diff --git a/Engine/Plugins/Animation/MLDeformer/ChaosClothGenerator/Source/ChaosClothGenerator/Private/ChaosClothGeneratorToolsMenuExtender.h b/Engine/Plugins/Animation/MLDeformer/ChaosClothGenerator/Source/ChaosClothGenerator/Private/ChaosClothGeneratorToolsMenuExtender.h
new file mode 100644
--- /dev/null
+++ b/Engine/Plugins/Animation/MLDeformer/ChaosClothGenerator/Source/ChaosClothGenerator/Private/ChaosClothGeneratorToolsMenuExtender.h
@@ -0,0 +1,32 @@
+// Copyright Epic Games, Inc. All Rights Reserved.
+
+#pragma once
+
+#include "CoreTypes.h"
+#include "MLDeformerEditorToolkit.h"
+#include "WorkflowOrientedApp/WorkflowTabFactory.h"
+
+namespace UE::Chaos::ClothGenerator
+{
+	struct FChaosClothGeneratorTabSummoner : public FWorkflowTabFactory
+	{
+	public:
+		static const FName TabID;
+		using FMLDeformerEditorToolkit = UE::MLDeformer::FMLDeformerEditorToolkit;
+	
+		FChaosClothGeneratorTabSummoner(const TSharedRef<FMLDeformerEditorToolkit>& InEditor);
+		virtual TSharedRef<SWidget> CreateTabBody(const FWorkflowTabSpawnInfo& Info) const override;
+		virtual TSharedPtr<SToolTip> CreateTabToolTipWidget(const FWorkflowTabSpawnInfo& Info) const override;
+	};
+	
+	class FChaosClothGeneratorToolsMenuExtender : public UE::MLDeformer::FToolsMenuExtender
+	{
+	public:
+		using FMLDeformerEditorToolkit = UE::MLDeformer::FMLDeformerEditorToolkit;
+		~FChaosClothGeneratorToolsMenuExtender() override = default;
+		virtual FMenuEntryParams GetMenuEntry(FMLDeformerEditorToolkit& Toolkit) const override;
+		virtual TSharedPtr<FWorkflowTabFactory> GetTabSummoner(const TSharedRef<FMLDeformerEditorToolkit>& Toolkit) const;
+	};
+
+	TUniquePtr<FChaosClothGeneratorToolsMenuExtender> CreateToolsMenuExtender();
+};
\ No newline at end of file
diff --git a/Engine/Plugins/Animation/MLDeformer/ChaosClothGenerator/Source/ChaosClothGenerator/Private/ClothGeneratorComponent.cpp b/Engine/Plugins/Animation/MLDeformer/ChaosClothGenerator/Source/ChaosClothGenerator/Private/ClothGeneratorComponent.cpp
new file mode 100644
--- /dev/null
+++ b/Engine/Plugins/Animation/MLDeformer/ChaosClothGenerator/Source/ChaosClothGenerator/Private/ClothGeneratorComponent.cpp
@@ -0,0 +1,61 @@
+// Copyright Epic Games, Inc. All Rights Reserved.
+
+#include "ClothGeneratorComponent.h"
+
+namespace UE::Chaos::ClothGenerator
+{
+    FClothGeneratorProxy::FClothGeneratorProxy(const UChaosClothComponent& InClothComponent)
+	    : FClothSimulationProxy(InClothComponent)
+    {	
+    }
+
+    FClothGeneratorProxy::~FClothGeneratorProxy() = default;
+};
+
+UClothGeneratorComponent::UClothGeneratorComponent(const FObjectInitializer& ObjectInitializer)
+    : Super(ObjectInitializer)
+{
+}
+
+UClothGeneratorComponent::UClothGeneratorComponent(FVTableHelper& Helper)
+    : Super(Helper)
+{
+}
+
+UClothGeneratorComponent::~UClothGeneratorComponent() = default;
+
+TSharedPtr<UE::Chaos::ClothAsset::FClothSimulationProxy> UClothGeneratorComponent::CreateClothSimulationProxy()
+{
+    TSharedPtr<FProxy> ProxyShared = MakeShared<FProxy>(*this);
+    Proxy = ProxyShared;
+    return ProxyShared;
+}
+
+void UClothGeneratorComponent::TickComponent(float DeltaTime, enum ELevelTick TickType, FActorComponentTickFunction* ThisTickFunction)
+{
+    UActorComponent::TickComponent(DeltaTime, TickType, ThisTickFunction);
+}
+
+
+TWeakPtr<UE::Chaos::ClothGenerator::FClothGeneratorProxy> UClothGeneratorComponent::GetProxy() const
+{
+    return Proxy;
+}
+
+void UClothGeneratorComponent::Pose(const TArray<FTransform>& InComponentSpaceTransforms)
+{
+	if (!ensure(InComponentSpaceTransforms.Num() == GetComponentSpaceTransforms().Num()))
+	{
+		return;
+	}
+	GetEditableComponentSpaceTransforms() = InComponentSpaceTransforms;
+	bNeedToFlipSpaceBaseBuffers = true;
+	FinalizeBoneTransform();
+
+	UpdateBounds();
+    if (IsInGameThread())
+    {
+        MarkRenderTransformDirty();
+        MarkRenderDynamicDataDirty();
+    }
+}
diff --git a/Engine/Plugins/Animation/MLDeformer/ChaosClothGenerator/Source/ChaosClothGenerator/Private/ClothGeneratorComponent.h b/Engine/Plugins/Animation/MLDeformer/ChaosClothGenerator/Source/ChaosClothGenerator/Private/ClothGeneratorComponent.h
new file mode 100644
--- /dev/null
+++ b/Engine/Plugins/Animation/MLDeformer/ChaosClothGenerator/Source/ChaosClothGenerator/Private/ClothGeneratorComponent.h
@@ -0,0 +1,53 @@
+// Copyright Epic Games, Inc. All Rights Reserved.
+
+#pragma once
+
+#include "ChaosClothAsset/ClothComponent.h"
+#include "ChaosClothAsset/ClothSimulationProxy.h"
+
+#include "ClothGeneratorComponent.generated.h"
+
+namespace UE::Chaos::ClothGenerator
+{
+	class FClothGeneratorProxy : public UE::Chaos::ClothAsset::FClothSimulationProxy
+    {
+    public:
+        using FClothSimulationProxy = UE::Chaos::ClothAsset::FClothSimulationProxy;
+        explicit FClothGeneratorProxy(const UChaosClothComponent& InClothComponent);
+        ~FClothGeneratorProxy();
+
+        using FClothSimulationProxy::Tick;
+        using FClothSimulationProxy::FillSimulationContext;
+        using FClothSimulationProxy::InitializeConfigs;
+        using FClothSimulationProxy::WriteSimulationData;
+    };
+};
+
+/**
+ * Cloth data generation component.
+ */
+UCLASS()
+class UClothGeneratorComponent : public UChaosClothComponent
+{
+	GENERATED_BODY()	
+public:
+	UClothGeneratorComponent(const FObjectInitializer& ObjectInitializer);
+	UClothGeneratorComponent(FVTableHelper& Helper);
+	~UClothGeneratorComponent();
+
+    using FProxy = UE::Chaos::ClothGenerator::FClothGeneratorProxy;
+    TWeakPtr<FProxy> GetProxy() const;
+    
+	/** Pose the cloth component using component space transforms. */
+    void Pose(const TArray<FTransform>& InComponentSpaceTransforms);
+protected:
+	//~ Begin UActorComponent Interface
+	virtual void TickComponent(float DeltaTime, enum ELevelTick TickType, FActorComponentTickFunction* ThisTickFunction) override;
+	//~ End UActorComponent Interface
+
+    /** Begin UChaosClothComponent Interface */
+	virtual TSharedPtr<UE::Chaos::ClothAsset::FClothSimulationProxy> CreateClothSimulationProxy() override;
+    /** End UChaosClothComponent Interface */
+private:
+    TWeakPtr<FProxy> Proxy;
+};
diff --git a/Engine/Plugins/Animation/MLDeformer/ChaosClothGenerator/Source/ChaosClothGenerator/Private/ClothGeneratorProperties.h b/Engine/Plugins/Animation/MLDeformer/ChaosClothGenerator/Source/ChaosClothGenerator/Private/ClothGeneratorProperties.h
new file mode 100644
--- /dev/null
+++ b/Engine/Plugins/Animation/MLDeformer/ChaosClothGenerator/Source/ChaosClothGenerator/Private/ClothGeneratorProperties.h
@@ -0,0 +1,59 @@
+// Copyright Epic Games, Inc. All Rights Reserved.
+
+#pragma once
+
+#include "CoreMinimal.h"
+
+#include "ChaosClothAsset/ClothAsset.h"
+#include "Animation/AnimSequence.h"
+#include "Engine/SkinnedAsset.h"
+#include "GeometryCache.h"
+
+#include "ClothGeneratorProperties.generated.h"
+
+UCLASS()
+class UClothGeneratorProperties : public UObject
+{
+	GENERATED_BODY()
+public:
+	/* Skeletal mesh that will be used in MLDeformer */
+	UPROPERTY(EditAnywhere, Category = Input)
+	TObjectPtr<USkinnedAsset> SkeletalMeshAsset;
+
+	/* Chaos cloth asset used in simulation. This should be different from the skeletal mesh asset. */
+	UPROPERTY(EditAnywhere, Category = Input)
+	TObjectPtr<UChaosClothAsset> ClothAsset;
+
+	/* Training poses. */
+	UPROPERTY(EditAnywhere, Category = Input)
+	TObjectPtr<UAnimSequence> AnimationSequence;
+
+	/* e.g. "0, 2, 5-10, 12-15". If left empty, all frames will be used */
+	UPROPERTY(EditAnywhere, Category = Input)
+	FString FramesToSimulate;
+
+	/* Output meshes */
+	UPROPERTY(EditAnywhere, Category = Output, meta = (EditCondition = "!bDebug"))
+	TObjectPtr<UGeometryCache> SimulatedCache;
+
+	/* Debug a single pose */
+	UPROPERTY(EditAnywhere, Category = Debug)
+	bool bDebug = false;
+
+	/* The frame to inspect */
+	UPROPERTY(EditAnywhere, Category = Debug, meta = (Min = 0, EditCondition = "bDebug"))
+	uint32 DebugFrame = 0;
+
+	/* The output sequence */
+	UPROPERTY(EditAnywhere, Category = Debug, meta = (EditCondition = "bDebug"))
+	TObjectPtr<UGeometryCache> DebugCache;
+
+	UPROPERTY(EditAnywhere, Category = "Simulation Settings", meta = (Min = 0))
+	float TimeStep = 1.f / 30;
+
+	UPROPERTY(EditAnywhere, Category = "Simulation Settings", meta = (Min = 0))
+	int32 NumSteps = 200;
+
+	UPROPERTY(EditAnywhere, Category = "Simulation Settings", meta = (Min = 1, EditCondition = "!bDebug"))
+	int32 NumThreads = 1;
+};
\ No newline at end of file
diff --git a/Engine/Plugins/Animation/MLDeformer/ChaosClothGenerator/Source/ChaosClothGenerator/Private/SClothGeneratorWidget.cpp b/Engine/Plugins/Animation/MLDeformer/ChaosClothGenerator/Source/ChaosClothGenerator/Private/SClothGeneratorWidget.cpp
new file mode 100644
--- /dev/null
+++ b/Engine/Plugins/Animation/MLDeformer/ChaosClothGenerator/Source/ChaosClothGenerator/Private/SClothGeneratorWidget.cpp
@@ -0,0 +1,191 @@
+#include "SClothGeneratorWidget.h"
+
+#include "AssetRegistry/AssetRegistryModule.h"
+#include "ClothGeneratorProperties.h"
+#include "ChaosClothGenerator.h"
+#include "ContentBrowserModule.h"
+#include "DetailCategoryBuilder.h"
+#include "DetailLayoutBuilder.h"
+#include "IContentBrowserSingleton.h"
+#include "IDetailsView.h"
+#include "Misc/FileHelper.h"
+#include "PropertyCustomizationHelpers.h"
+#include "PropertyEditorModule.h"
+#include "Widgets/Input/SButton.h"
+#include "Widgets/Layout/SBox.h"
+#include "Widgets/Text/STextBlock.h"
+
+#define LOCTEXT_NAMESPACE "ClothGeneratorWidget"
+
+namespace UE::Chaos::ClothGenerator
+{
+	void SClothGeneratorWidget::Construct(const FArguments& InArgs)
+	{
+		FPropertyEditorModule& PropertyModule = FModuleManager::LoadModuleChecked<FPropertyEditorModule>("PropertyEditor");
+		ChaosClothGenerator = MakeShared<FChaosClothGenerator>();
+
+		FDetailsViewArgs Args;
+		DetailsView = PropertyModule.CreateDetailView(Args);
+		DetailsView->SetObject(&ChaosClothGenerator->GetProperties());
+
+		this->ChildSlot
+		[
+			SNew(SVerticalBox)
+			+ SVerticalBox::Slot()
+			.AutoHeight()
+			[
+				DetailsView.ToSharedRef()
+			]
+			+ SVerticalBox::Slot()
+			.AutoHeight()
+			.HAlign(HAlign_Center)
+			[
+				SNew(SBox)
+				.WidthOverride(300)
+				[
+					SNew(SButton)
+					.Text(FText::FromString("Start Generating"))
+					.HAlign(HAlign_Center)
+					.OnClicked_Lambda([this]() -> FReply
+					{
+						ChaosClothGenerator->RequestAction(EClothGeneratorActions::StartGenerate);
+						return FReply::Handled();
+					})
+				]
+			]
+		];
+	}
+
+	TWeakObjectPtr<UClothGeneratorProperties> SClothGeneratorWidget::GetProperties() const
+	{
+		if (ChaosClothGenerator.IsValid())
+		{
+			return &ChaosClothGenerator->GetProperties();
+		}
+		return nullptr;
+	}
+
+	TSharedRef<IDetailCustomization> FClothGeneratorDetails::MakeInstance()
+	{
+		return MakeShareable(new FClothGeneratorDetails);
+	}
+
+	namespace Private
+	{
+		template<class T>
+		T* CreateOrLoad(const FString& PackageName)
+		{
+			const FName AssetName(FPackageName::GetLongPackageAssetName(PackageName));		
+			if (UPackage* const Package = CreatePackage(*PackageName))
+			{
+				LoadPackage(nullptr, *PackageName, LOAD_Quiet | LOAD_EditorOnly);
+				T* Asset = FindObject<T>(Package, *AssetName.ToString());
+				if (!Asset)
+				{
+					Asset = NewObject<T>(Package, *AssetName.ToString(), RF_Public | RF_Standalone | RF_Transactional);
+					Asset->MarkPackageDirty();
+					FAssetRegistryModule::AssetCreated(Asset);
+				}
+				return Asset;
+			}
+			return nullptr;
+		}
+
+		TObjectPtr<UGeometryCache> NewGeometryCacheDialog(const UObject* NamingAsset = nullptr)
+		{
+			FSaveAssetDialogConfig Config;
+			{
+				if (NamingAsset)
+				{
+					const FString PackageName = NamingAsset->GetOutermost()->GetName();
+					Config.DefaultPath = FPackageName::GetLongPackagePath(PackageName);
+					Config.DefaultAssetName = FString::Printf(TEXT("GC_%s"), *NamingAsset->GetName());
+				}
+				Config.AssetClassNames.Add(UGeometryCache::StaticClass()->GetClassPathName());
+				Config.ExistingAssetPolicy = ESaveAssetDialogExistingAssetPolicy::Disallow;
+				Config.DialogTitleOverride = LOCTEXT("ExportGeometryCacheDialogTitle", "Export Geometry Cache As");
+			}
+				
+			FContentBrowserModule& ContentBrowserModule = FModuleManager::LoadModuleChecked<FContentBrowserModule>("ContentBrowser");
+				
+			FString NewPackageName;
+			FText OutError;
+			for (bool bFilenameValid = false; !bFilenameValid; bFilenameValid = FFileHelper::IsFilenameValidForSaving(NewPackageName, OutError))
+			{
+				const FString AssetPath = ContentBrowserModule.Get().CreateModalSaveAssetDialog(Config);
+				if (AssetPath.IsEmpty())
+				{
+					return nullptr;
+				}
+				NewPackageName = FPackageName::ObjectPathToPackageName(AssetPath);
+			}
+			return CreateOrLoad<UGeometryCache>(NewPackageName);
+		}
+
+		TWeakObjectPtr<UClothGeneratorProperties> GetProperties(IDetailLayoutBuilder& DetailBuilder)
+		{
+			TArray<TWeakObjectPtr<UObject>> ObjectsBeingCustomized;
+			DetailBuilder.GetObjectsBeingCustomized(ObjectsBeingCustomized);
+			if (ObjectsBeingCustomized.IsEmpty())
+			{
+				return nullptr;
+			}
+			else
+			{
+				return Cast<UClothGeneratorProperties>(ObjectsBeingCustomized[0]);
+			}
+		}
+
+		void AddGeometryCacheRowWithButton(IDetailLayoutBuilder& DetailBuilder, IDetailCategoryBuilder& CategoryBuilder, const FName PropertyName)
+		{
+			TSharedPtr<IPropertyHandle> Property = DetailBuilder.GetProperty(PropertyName);
+			DetailBuilder.HideProperty(Property);
+			CategoryBuilder.AddCustomRow(FText::FromName(PropertyName))
+			.NameContent()
+			[
+				Property->CreatePropertyNameWidget()
+			]
+			.ValueContent()
+			[
+				SNew(SHorizontalBox)
+				+ SHorizontalBox::Slot()
+				.AutoWidth()
+				[
+					Property->CreatePropertyValueWidget()
+				]
+				+ SHorizontalBox::Slot()
+				.MaxWidth(100)
+				.VAlign(VAlign_Center)
+				[
+					SNew(SButton)
+					.Text(FText::FromString("New"))
+					.HAlign(HAlign_Center)
+					.VAlign(VAlign_Center)
+					.OnClicked_Lambda([Property, &DetailBuilder]() -> FReply
+					{
+						TWeakObjectPtr<UClothGeneratorProperties> Properties = Private::GetProperties(DetailBuilder);
+						const UObject* const NamingAsset = Properties.IsValid() ? Properties->SkeletalMeshAsset.Get() : nullptr;
+						UGeometryCache* const NewGeometryCache = Private::NewGeometryCacheDialog(NamingAsset);
+						Property->SetValue(NewGeometryCache);
+						return FReply::Handled();
+					})
+				]
+			];
+		}
+	}
+
+
+	void FClothGeneratorDetails::CustomizeDetails(IDetailLayoutBuilder& DetailBuilder)
+	{
+		IDetailCategoryBuilder& InputCategory = DetailBuilder.EditCategory("Input");
+		IDetailCategoryBuilder& OutputCategory = DetailBuilder.EditCategory("Output");
+		Private::AddGeometryCacheRowWithButton(DetailBuilder, OutputCategory, GET_MEMBER_NAME_CHECKED(UClothGeneratorProperties, SimulatedCache));
+
+		IDetailCategoryBuilder& DebugCategory = DetailBuilder.EditCategory("Debug");
+		DebugCategory.AddProperty(GET_MEMBER_NAME_CHECKED(UClothGeneratorProperties, bDebug));
+		DebugCategory.AddProperty(GET_MEMBER_NAME_CHECKED(UClothGeneratorProperties, DebugFrame));
+		Private::AddGeometryCacheRowWithButton(DetailBuilder, DebugCategory, GET_MEMBER_NAME_CHECKED(UClothGeneratorProperties, DebugCache));
+	}
+};
+
+#undef LOCTEXT_NAMESPACE
\ No newline at end of file
diff --git a/Engine/Plugins/Animation/MLDeformer/ChaosClothGenerator/Source/ChaosClothGenerator/Private/SClothGeneratorWidget.h b/Engine/Plugins/Animation/MLDeformer/ChaosClothGenerator/Source/ChaosClothGenerator/Private/SClothGeneratorWidget.h
new file mode 100644
--- /dev/null
+++ b/Engine/Plugins/Animation/MLDeformer/ChaosClothGenerator/Source/ChaosClothGenerator/Private/SClothGeneratorWidget.h
@@ -0,0 +1,37 @@
+// Copyright Epic Games, Inc. All Rights Reserved.
+
+#pragma once
+
+#include "IDetailCustomization.h"
+#include "Widgets/SCompoundWidget.h"
+
+class UClothGeneratorProperties;
+class IDetailsView;
+namespace UE::Chaos::ClothGenerator
+{
+	class FChaosClothGenerator;
+
+	class SClothGeneratorWidget : public SCompoundWidget
+	{
+		SLATE_BEGIN_ARGS(SClothGeneratorWidget) 
+		{}
+		SLATE_END_ARGS()
+
+	public:
+		void Construct(const FArguments& InArgs);
+		TWeakObjectPtr<UClothGeneratorProperties> GetProperties() const;
+
+	private:
+		TSharedPtr<IDetailsView> DetailsView;
+		TSharedPtr<FChaosClothGenerator> ChaosClothGenerator;
+	};
+
+	class FClothGeneratorDetails : public IDetailCustomization
+	{
+	public:
+		static TSharedRef<IDetailCustomization> MakeInstance();
+		virtual void CustomizeDetails(IDetailLayoutBuilder& DetailBuilder) override;
+	private:
+		TWeakObjectPtr<UClothGeneratorProperties> GetProperties(IDetailLayoutBuilder& DetailBuilder) const;
+	};
+};
\ No newline at end of file
diff --git a/Engine/Plugins/Animation/MLDeformer/MLDeformerFramework/Source/MLDeformerFrameworkEditor/Private/MLDeformerApplicationMode.cpp b/Engine/Plugins/Animation/MLDeformer/MLDeformerFramework/Source/MLDeformerFrameworkEditor/Private/MLDeformerApplicationMode.cpp
--- a/Engine/Plugins/Animation/MLDeformer/MLDeformerFramework/Source/MLDeformerFrameworkEditor/Private/MLDeformerApplicationMode.cpp
+++ b/Engine/Plugins/Animation/MLDeformer/MLDeformerFramework/Source/MLDeformerFrameworkEditor/Private/MLDeformerApplicationMode.cpp
@@ -40,6 +40,13 @@ namespace UE::MLDeformer
 		// Register custom tabs.
 		TabFactories.RegisterFactory(MakeShared<FMLDeformerVizSettingsTabSummoner>(MLDeformerEditor));
 		TabFactories.RegisterFactory(MakeShared<FMLDeformerTimelineTabSummoner>(MLDeformerEditor));
+		for (const TUniquePtr<FToolsMenuExtender>& Extender : FMLDeformerEditorToolkit::GetToolsMenuExtenders())
+		{
+			if (TSharedPtr<FWorkflowTabFactory> Factory = Extender->GetTabSummoner(MLDeformerEditor); Factory.IsValid())
+			{
+				TabFactories.RegisterFactory(Factory.ToSharedRef());
+			}
+		}
 
 		// Create tab layout.
 		TabLayout = FTabManager::NewLayout("Standalone_MLDeformerEditor_Layout_v5.3.001")
diff --git a/Engine/Plugins/Animation/MLDeformer/MLDeformerFramework/Source/MLDeformerFrameworkEditor/Private/MLDeformerEditorToolkit.cpp b/Engine/Plugins/Animation/MLDeformer/MLDeformerFramework/Source/MLDeformerFrameworkEditor/Private/MLDeformerEditorToolkit.cpp
--- a/Engine/Plugins/Animation/MLDeformer/MLDeformerFramework/Source/MLDeformerFrameworkEditor/Private/MLDeformerEditorToolkit.cpp
+++ b/Engine/Plugins/Animation/MLDeformer/MLDeformerFramework/Source/MLDeformerFrameworkEditor/Private/MLDeformerEditorToolkit.cpp
@@ -36,6 +36,9 @@ namespace UE::MLDeformer
 	const FName MLDeformerEditorModes::Editor("MLDeformerEditorMode");
 	const FName MLDeformerEditorAppName = FName(TEXT("MLDeformerEditorApp"));
 
+	TArray<TUniquePtr<FToolsMenuExtender>> FMLDeformerEditorToolkit::ToolsMenuExtenders;
+	FCriticalSection FMLDeformerEditorToolkit::ExtendersMutex;
+
 	FMLDeformerEditorToolkit::~FMLDeformerEditorToolkit()
 	{
 		ActiveModel.Reset();
@@ -233,6 +236,19 @@ namespace UE::MLDeformer
 		return MenuBuilder.MakeWidget();
 	}
 
+	TSharedRef<SWidget> FMLDeformerEditorToolkit::GenerateToolsMenuContents(TSharedRef<FUICommandList> InCommandList)
+	{
+		const bool bShouldCloseWindowAfterMenuSelection = true;
+		FMenuBuilder MenuBuilder(bShouldCloseWindowAfterMenuSelection, InCommandList);
+
+		for (const TUniquePtr<FToolsMenuExtender>& Extender : ToolsMenuExtenders)
+		{
+			MenuBuilder.AddMenuEntry(Extender->GetMenuEntry(*this));
+		}
+
+		return MenuBuilder.MakeWidget();
+	}
+
 	void FMLDeformerEditorToolkit::SwitchVizMode(EMLDeformerVizMode Mode)
 	{
 		if (!ActiveModel.IsValid())
@@ -482,6 +498,27 @@ namespace UE::MLDeformer
 			);
 		}
 		ToolbarBuilder.EndSection();
+
+		if (!ToolsMenuExtenders.IsEmpty())
+		{
+			ToolbarBuilder.BeginSection("Tools");
+			{
+				TSharedPtr<FUICommandList> CommandList = GetToolkitCommands();
+				ToolbarBuilder.AddComboButton(
+					FUIAction(),
+					FOnGetContent::CreateRaw(this, &FMLDeformerEditorToolkit::GenerateToolsMenuContents, CommandList.ToSharedRef()),
+					TAttribute<FText>::CreateLambda(
+						[this]()
+						{
+							return LOCTEXT("ToolsMenu", "Tools");
+						}
+					),
+					LOCTEXT("ToolsMenuTooltip", "Tools"),
+					FSlateIcon(FMLDeformerEditorStyle::Get().GetStyleSetName(), "MLDeformer.VizSettings.TabIcon")
+				);
+			}
+			ToolbarBuilder.EndSection();
+		}
 	}
 
 	bool FMLDeformerEditorToolkit::HandleTrainingResult(ETrainingResult TrainingResult, double TrainingDuration, bool& bOutUsePartiallyTrained, bool bSuppressDialogs, bool& bOutSuccess)
@@ -904,6 +941,18 @@ namespace UE::MLDeformer
 		OnModelChanged(Index, bForceChange);
 		return true;
 	}
+
+	void FMLDeformerEditorToolkit::AddToolsMenuExtender(TUniquePtr<FToolsMenuExtender> Extender)
+	{
+		FScopeLock Lock(&ExtendersMutex);
+		ToolsMenuExtenders.Emplace_GetRef(MoveTemp(Extender));
+	}
+
+	TConstArrayView<TUniquePtr<FToolsMenuExtender>> FMLDeformerEditorToolkit::GetToolsMenuExtenders()
+	{
+		FScopeLock Lock(&ExtendersMutex);
+		return ToolsMenuExtenders;
+	}
 }	// namespace UE::MLDeformer
 
 #undef LOCTEXT_NAMESPACE
diff --git a/Engine/Plugins/Animation/MLDeformer/MLDeformerFramework/Source/MLDeformerFrameworkEditor/Public/MLDeformerEditorToolkit.h b/Engine/Plugins/Animation/MLDeformer/MLDeformerFramework/Source/MLDeformerFrameworkEditor/Public/MLDeformerEditorToolkit.h
--- a/Engine/Plugins/Animation/MLDeformer/MLDeformerFramework/Source/MLDeformerFrameworkEditor/Public/MLDeformerEditorToolkit.h
+++ b/Engine/Plugins/Animation/MLDeformer/MLDeformerFramework/Source/MLDeformerFrameworkEditor/Public/MLDeformerEditorToolkit.h
@@ -11,9 +11,12 @@
 #include "IPersonaViewport.h"
 #include "PersonaAssetEditorToolkit.h"
 #include "EditorUndoClient.h"
+#include "Framework/MultiBox/MultiBox.h"
 #include "Math/Color.h"
 #include "Widgets/Notifications/SNotificationList.h"
 
+struct FMenuEntryParams;
+class FTabManager;
 class IDetailsView;
 class UMLDeformerAsset;
 class SSimpleTimeSlider;
@@ -28,6 +31,15 @@ namespace UE::MLDeformer
 		extern const FName Editor;
 	}
 
+	class FMLDeformerEditorToolkit;
+	class MLDEFORMERFRAMEWORKEDITOR_API FToolsMenuExtender
+	{
+	public:
+		virtual ~FToolsMenuExtender() {}
+		virtual FMenuEntryParams GetMenuEntry(FMLDeformerEditorToolkit& Toolkit) const = 0;
+		virtual TSharedPtr<FWorkflowTabFactory> GetTabSummoner(const TSharedRef<FMLDeformerEditorToolkit>& Toolkit) const { return nullptr; }
+	};
+
 	/**
 	 * The ML Deformer asset editor toolkit.
 	 * This is the editor that opens when you double click an ML Deformer asset.
@@ -117,6 +129,8 @@ namespace UE::MLDeformer
 		bool IsTrainButtonEnabled() const;
 		bool IsTraining() const;
 
+		static void AddToolsMenuExtender(TUniquePtr<FToolsMenuExtender> Extender);
+		static TConstArrayView<TUniquePtr<FToolsMenuExtender>> GetToolsMenuExtenders();
 	private:
 		UE_DEPRECATED(5.3, "Please use the OnModelChanged that takes two parameters instead.")
 		void OnModelChanged(int Index);
@@ -148,6 +162,7 @@ namespace UE::MLDeformer
 
 		TSharedRef<SWidget> GenerateModelButtonContents(TSharedRef<FUICommandList> InCommandList);
 		TSharedRef<SWidget> GenerateVizModeButtonContents(TSharedRef<FUICommandList> InCommandList);
+		TSharedRef<SWidget> GenerateToolsMenuContents(TSharedRef<FUICommandList> InCommandList);
 
 	private:
 		/** The persona toolkit. */	
@@ -178,5 +193,11 @@ namespace UE::MLDeformer
 
 		/** Are we currently in a training process? */
 		bool bIsTraining = false;
+
+		/** Extenders for Tools menu */
+		static TArray<TUniquePtr<FToolsMenuExtender>> ToolsMenuExtenders;
+
+		/** Mutex for adding extenders */
+		static FCriticalSection ExtendersMutex;
 	};
 }	// namespace UE::MLDeformer
