ChaosClothGenerator: fix tool tip conflict

diff --git a/Engine/Plugins/Animation/MLDeformer/ChaosClothGenerator/Source/ChaosClothGenerator/Private/ChaosClothGeneratorToolsMenuExtender.cpp b/Engine/Plugins/Animation/MLDeformer/ChaosClothGenerator/Source/ChaosClothGenerator/Private/ChaosClothGeneratorToolsMenuExtender.cpp
--- a/Engine/Plugins/Animation/MLDeformer/ChaosClothGenerator/Source/ChaosClothGenerator/Private/ChaosClothGeneratorToolsMenuExtender.cpp
+++ b/Engine/Plugins/Animation/MLDeformer/ChaosClothGenerator/Source/ChaosClothGenerator/Private/ChaosClothGeneratorToolsMenuExtender.cpp
@@ -86,7 +86,7 @@ namespace UE::Chaos::ClothGenerator
 	TSharedPtr<SToolTip> FChaosClothGeneratorTabSummoner::CreateTabToolTipWidget(const FWorkflowTabSpawnInfo& Info) const
 	{
 		return IDocumentation::Get()->CreateToolTip(
-			LOCTEXT("ChaosClothGeneratorTooltip", "Generate training data using chaos cloth solver."),
+			LOCTEXT("ChaosClothGeneratorWidgetTooltip", "Generate training data using chaos cloth solver."),
 			NULL,
 			TEXT("Shared/Editors/Persona"),
 			TEXT("ChaosClothGenerator_Window"));
@@ -106,7 +106,7 @@ namespace UE::Chaos::ClothGenerator
 			})
 		);
 		Params.LabelOverride = LOCTEXT("ChaosClothGenerator", "Chaos Cloth Generator");
-		Params.ToolTipOverride = LOCTEXT("ChaosClothGeneratorTooltip", "Generate training data using chaos cloth solver");
+		Params.ToolTipOverride = LOCTEXT("ChaosClothGeneratorMenuTooltip", "Generate training data using chaos cloth solver");
 		return Params;
 	}
 	
