ChaosClothGenerator: fix root transform misalignment

diff --git a/Engine/Plugins/Animation/MLDeformer/ChaosClothGenerator/Source/ChaosClothGenerator/Private/ChaosClothGenerator.cpp b/Engine/Plugins/Animation/MLDeformer/ChaosClothGenerator/Source/ChaosClothGenerator/Private/ChaosClothGenerator.cpp
--- a/Engine/Plugins/Animation/MLDeformer/ChaosClothGenerator/Source/ChaosClothGenerator/Private/ChaosClothGenerator.cpp
+++ b/Engine/Plugins/Animation/MLDeformer/ChaosClothGenerator/Source/ChaosClothGenerator/Private/ChaosClothGenerator.cpp
@@ -587,7 +587,8 @@ namespace UE::Chaos::ClothGenerator
 	
 		FAnimationPoseData AnimationPoseData(OutPose, OutCurve, TempAttributes);
 		AnimationSequence->GetAnimationPose(AnimationPoseData, ExtractionContext);
-	
+
+		const FTransform RootTransform = AnimationSequence->ExtractRootTrackTransform(Time, nullptr);
 		TArray<FTransform> ComponentSpaceTransforms;
 		ComponentSpaceTransforms.SetNumUninitialized(NumBones);
 		for (int32 Index = 0; Index < NumBones; ++Index)
@@ -597,7 +598,7 @@ namespace UE::Chaos::ClothGenerator
 			ComponentSpaceTransforms[Index] = 
 				ComponentSpaceTransforms.IsValidIndex(ParentIndex) && ParentIndex < Index ? 
 				AnimationPoseData.GetPose()[CompactIndex] * ComponentSpaceTransforms[ParentIndex] : 
-				ReferenceSkeleton->GetRefBonePose()[Index];
+				RootTransform;
 		}
 	
 		return ComponentSpaceTransforms;
