ChaosClothGenerator: fix problem when Reference bones have different bone orders

diff --git a/Engine/Plugins/Animation/MLDeformer/ChaosClothGenerator/Source/ChaosClothGenerator/Private/ChaosClothGenerator.cpp b/Engine/Plugins/Animation/MLDeformer/ChaosClothGenerator/Source/ChaosClothGenerator/Private/ChaosClothGenerator.cpp
--- a/Engine/Plugins/Animation/MLDeformer/ChaosClothGenerator/Source/ChaosClothGenerator/Private/ChaosClothGenerator.cpp
+++ b/Engine/Plugins/Animation/MLDeformer/ChaosClothGenerator/Source/ChaosClothGenerator/Private/ChaosClothGenerator.cpp
@@ -2,6 +2,7 @@
 
 #include "ChaosClothGenerator.h"
 
+#include "AnimationRuntime.h"
 #include "Animation/AnimSequence.h"
 #include "Animation/AttributesRuntime.h"
 #include "AssetRegistry/AssetRegistryModule.h"
@@ -589,18 +590,20 @@ namespace UE::Chaos::ClothGenerator
 		AnimationSequence->GetAnimationPose(AnimationPoseData, ExtractionContext);
 
 		const FTransform RootTransform = AnimationSequence->ExtractRootTrackTransform(Time, nullptr);
-		TArray<FTransform> ComponentSpaceTransforms;
-		ComponentSpaceTransforms.SetNumUninitialized(NumBones);
+
+		TArray<FTransform> BoneTransforms;
+		BoneTransforms.SetNumZeroed(NumBones);
+		const FSkeletonToMeshLinkup& LinkupTable = Skeleton->FindOrAddMeshLinkupData(ClothAsset);
+		const TArray<int32>& BoneMap = LinkupTable.SkeletonToMeshTable;
+		check(BoneMap.Num() == NumBones);
 		for (int32 Index = 0; Index < NumBones; ++Index)
 		{
 			const FCompactPoseBoneIndex CompactIndex = BoneContainer.MakeCompactPoseIndex(FMeshPoseBoneIndex(Index));
-			const int32 ParentIndex = ReferenceSkeleton->GetParentIndex(Index);
-			ComponentSpaceTransforms[Index] = 
-				ComponentSpaceTransforms.IsValidIndex(ParentIndex) && ParentIndex < Index ? 
-				AnimationPoseData.GetPose()[CompactIndex] * ComponentSpaceTransforms[ParentIndex] : 
-				RootTransform;
+			check(BoneTransforms.IsValidIndex(BoneMap[Index]));
+			BoneTransforms[BoneMap[Index]] = AnimationPoseData.GetPose()[CompactIndex];
 		}
-	
+		TArray<FTransform> ComponentSpaceTransforms;
+		FAnimationRuntime::FillUpComponentSpaceTransforms(*ReferenceSkeleton, BoneTransforms, ComponentSpaceTransforms);
 		return ComponentSpaceTransforms;
 	}
 	
