ChaosClothGenerator: fix windows deferencing null pointer issue

diff --git a/Engine/Plugins/Animation/MLDeformer/ChaosClothGenerator/Source/ChaosClothGenerator/Private/ChaosClothGenerator.cpp b/Engine/Plugins/Animation/MLDeformer/ChaosClothGenerator/Source/ChaosClothGenerator/Private/ChaosClothGenerator.cpp
--- a/Engine/Plugins/Animation/MLDeformer/ChaosClothGenerator/Source/ChaosClothGenerator/Private/ChaosClothGenerator.cpp
+++ b/Engine/Plugins/Animation/MLDeformer/ChaosClothGenerator/Source/ChaosClothGenerator/Private/ChaosClothGenerator.cpp
@@ -568,6 +568,11 @@ namespace UE::Chaos::ClothGenerator
 		const FReferenceSkeleton* const ReferenceSkeleton = ClothAsset ? &ClothAsset->GetRefSkeleton() : nullptr;
 		USkeleton* const Skeleton = ClothAsset ? ClothAsset->GetSkeleton() : nullptr;
 		const int32 NumBones = ReferenceSkeleton ? ReferenceSkeleton->GetNum() : 0;
+
+		if (!ReferenceSkeleton || NumBones == 0)
+		{
+			return TArray<FTransform>();
+		}
 	
 		TArray<uint16> BoneIndices;
 		BoneIndices.SetNumUninitialized(NumBones);
