ChaosClothGenerator: support skeletal mesh with multiple parts

diff --git a/Engine/Plugins/Animation/MLDeformer/ChaosClothGenerator/Source/ChaosClothGenerator/ChaosClothGenerator.build.cs b/Engine/Plugins/Animation/MLDeformer/ChaosClothGenerator/Source/ChaosClothGenerator/ChaosClothGenerator.build.cs
--- a/Engine/Plugins/Animation/MLDeformer/ChaosClothGenerator/Source/ChaosClothGenerator/ChaosClothGenerator.build.cs
+++ b/Engine/Plugins/Animation/MLDeformer/ChaosClothGenerator/Source/ChaosClothGenerator/ChaosClothGenerator.build.cs
@@ -38,10 +38,12 @@ public class ChaosClothGenerator : ModuleRules
 				"ChaosClothAssetEngine",
 				"Engine",
 				"GeometryCache",
+				"MeshDescription",
 				"MLDeformerFramework",
 				"MLDeformerFrameworkEditor",
 				"PropertyEditor",
 				"RenderCore",
+				"SkeletalMeshDescription",
 				"Slate",
 				"SlateCore",
 				"UnrealEd"
diff --git a/Engine/Plugins/Animation/MLDeformer/ChaosClothGenerator/Source/ChaosClothGenerator/Private/ChaosClothGenerator.cpp b/Engine/Plugins/Animation/MLDeformer/ChaosClothGenerator/Source/ChaosClothGenerator/Private/ChaosClothGenerator.cpp
--- a/Engine/Plugins/Animation/MLDeformer/ChaosClothGenerator/Source/ChaosClothGenerator/Private/ChaosClothGenerator.cpp
+++ b/Engine/Plugins/Animation/MLDeformer/ChaosClothGenerator/Source/ChaosClothGenerator/Private/ChaosClothGenerator.cpp
@@ -11,6 +11,7 @@
 #include "ChaosClothAsset/ClothSimulationProxy.h"
 #include "ClothGeneratorComponent.h"
 #include "ClothGeneratorProperties.h"
+#include "Engine/SkeletalMesh.h"
 #include "Engine/SkinnedAssetCommon.h"
 #include "FileHelpers.h"
 #include "GeometryCache.h"
@@ -25,6 +26,7 @@
 #include "PhysicsEngine/PhysicsAsset.h"
 #include "Rendering/SkeletalMeshRenderData.h"
 #include "Rendering/SkeletalMeshModel.h"
+#include "SkeletalMeshAttributes.h"
 #include "SkeletalRenderPublic.h"
 #include "Tasks/Pipe.h"
 #include "UObject/SavePackage.h"
@@ -121,6 +123,18 @@ namespace UE::Chaos::ClothGenerator
 			}
 			return NumVertices;
 		}
+
+		int32 GetNumVertices(const USkeletalMesh& SkeletalMesh)
+		{
+			const FSkeletalMeshRenderData* RenderData = SkeletalMesh.GetResourceForRendering();
+			constexpr int32 LODIndex = 0;
+			if (!RenderData || !RenderData->LODRenderData.IsValidIndex(LODIndex))
+			{
+				return INDEX_NONE;
+			}
+			const FSkeletalMeshLODRenderData& LODData = RenderData->LODRenderData[LODIndex];
+			return GetNumVertices(LODData);
+		}
 		
 		TArrayView<TArray<FVector3f>> ShrinkToValidFrames(const TArrayView<TArray<FVector3f>>& Positions, int32 NumVertices)
 		{
@@ -135,31 +149,382 @@ namespace UE::Chaos::ClothGenerator
 			}
 			return TArrayView<TArray<FVector3f>>(Positions.GetData(), NumValidFrames);
 		}
-		
-		void SaveGeometryCache(UGeometryCache& GeometryCache, const USkinnedAsset& Asset, TConstArrayView<uint32> ImportedVertexNumbers, TArrayView<TArray<FVector3f>> PositionsToMoveFrom)
+
+		TArray<FVector2f> GetUV0s(const FSkeletalMeshLODRenderData& LODData)
+		{
+			TArray<FVector2f> UV0s;
+			const FStaticMeshVertexBuffer& StaticMeshVertexBuffer = LODData.StaticVertexBuffers.StaticMeshVertexBuffer;
+			const int32 NumVertices = StaticMeshVertexBuffer.GetNumVertices();
+			UV0s.SetNumZeroed(NumVertices);
+			for (int32 Index = 0; Index < NumVertices; ++Index)
+			{
+				UV0s[Index] = StaticMeshVertexBuffer.GetVertexUV(Index, 0);
+			}
+			return UV0s;
+		}
+
+		TArray<FColor> GetColors(const FSkeletalMeshLODRenderData& LODData, int32 NumVertices)
+		{
+			TArray<FColor> Colors;
+			Colors.SetNum(NumVertices);
+			const FColorVertexBuffer& ColorVertexBuffer = LODData.StaticVertexBuffers.ColorVertexBuffer;
+			if (ColorVertexBuffer.GetNumVertices() == NumVertices)
+			{
+				for (int32 Index = 0; Index < NumVertices; ++Index)
+				{
+					Colors[Index] = ColorVertexBuffer.VertexColor(Index);
+				}
+			}
+			else
+			{
+				for (int32 Index = 0; Index < NumVertices; ++Index)
+				{
+					Colors[Index] = FColor::White;
+				}
+			}
+			return Colors;
+		}
+
+		TArray<TObjectPtr<UMaterialInterface>> GetMaterialInterfaces(const USkinnedAsset& Asset)
+		{
+			const TArray<FSkeletalMaterial>& Materials = Asset.GetMaterials();
+			TArray<TObjectPtr<UMaterialInterface>> Interfaces;
+			Interfaces.SetNum(Materials.Num());
+			for (int32 Index = 0; Index < Materials.Num(); ++Index)
+			{
+				Interfaces[Index] = Materials[Index].MaterialInterface;
+			}
+			return Interfaces;
+		}
+
+		TArray<TInterval<int32>> GetImportPartIntervals(FSkeletalMeshConstAttributes MeshAttributes)
+		{
+			const int32 NumParts = MeshAttributes.GetNumSourceGeometryParts();
+			const FSkeletalMeshAttributes::FSourceGeometryPartVertexOffsetAndCountConstRef OffsetAndCounts = MeshAttributes.GetSourceGeometryPartVertexOffsetAndCounts();
+			TArray<TInterval<int32>> Result;
+			Result.SetNum(NumParts);
+			for (int32 PartIndex = 0; PartIndex < NumParts; ++PartIndex)
+			{
+				TConstArrayView<int32> OffsetAndCount = OffsetAndCounts.Get(PartIndex);
+				const int32 Offset = OffsetAndCount[0];
+				const int32 Count = OffsetAndCount[1];
+				Result[PartIndex] = TInterval<int32>(Offset, Offset + Count - 1);
+			}
+			return Result;
+		}
+
+		int32 GetPartIndexFromVertexIndex(int32 VertexIndex, const TArray<int32>& MeshToImportVertexMap, const TArray<TInterval<int32>>& ImportPartIntervals)
+		{
+			if (!MeshToImportVertexMap.IsValidIndex(VertexIndex))
+			{
+				return INDEX_NONE;
+			}
+			const int32 ImportIndex = MeshToImportVertexMap[VertexIndex];
+			const int32 NumParts = ImportPartIntervals.Num();
+			for (int32 PartIndex = 0; PartIndex < NumParts; ++PartIndex)
+			{
+				const TInterval<int32>& Interval = ImportPartIntervals[PartIndex];
+				if (Interval.Contains(ImportIndex))
+				{
+					return PartIndex;
+				}
+			}
+			return INDEX_NONE;
+		}
+
+		TArray<int32> GetVertexToPart(const TArray<int32>& MeshToImportVertexMap, const TArray<TInterval<int32>>& ImportPartIntervals)
+		{
+			const int32 NumVertices = MeshToImportVertexMap.Num();
+			TArray<int32> VertexToPart;
+			VertexToPart.SetNum(NumVertices);
+			for (int32 VertexIndex = 0; VertexIndex < NumVertices; ++VertexIndex)
+			{
+				VertexToPart[VertexIndex] = GetPartIndexFromVertexIndex(VertexIndex, MeshToImportVertexMap, ImportPartIntervals);
+			}
+			return VertexToPart;
+		}
+
+		TArray<TArray<int32>> GetPartToIndices(const TArray<int32>& PartIndices, const int32 NumParts)
+		{
+			TArray<TArray<int32>> PartToIndices;
+			PartToIndices.SetNum(NumParts);
+			for (int32 Index = 0; Index < PartIndices.Num(); ++Index)
+			{
+				const int32 PartIndex = PartIndices[Index];
+				PartToIndices[PartIndex].Add(Index);
+			}
+			return PartToIndices;
+		}
+
+		TOptional<TArray<int32>> GetTriangleToPart(const TArray<uint32>& Indices, const TArray<int32>& VertexToPart)
+		{
+			check(Indices.Num() % 3 == 0);
+			const int32 NumTriangles = Indices.Num() / 3;
+			TOptional<TArray<int32>> None;
+			TArray<int32> TriangleToPart;
+			TriangleToPart.SetNum(NumTriangles);
+			for (int32 TriangleIndex = 0; TriangleIndex < NumTriangles; ++TriangleIndex)
+			{
+				int32 PartIndex = INDEX_NONE;
+				for (int32 Corner = 0; Corner < 3; ++Corner)
+				{
+					const int32 VertexIndex = Indices[TriangleIndex * 3 + Corner];
+					const int32 CandidatePartIndex = VertexToPart[VertexIndex];
+					if (PartIndex == INDEX_NONE)
+					{
+						PartIndex = CandidatePartIndex;
+					}
+					else if (PartIndex != CandidatePartIndex)
+					{
+						return None;
+					}
+				}
+				TriangleToPart[TriangleIndex] = PartIndex;
+			}
+			return TriangleToPart;
+		}
+
+		TArray<int32> GetIndicesFromTriangles(const TArray<int32>& Triangles)
+		{
+			TArray<int32> Indices;
+			Indices.SetNum(Triangles.Num() * 3);
+			for (int32 TriangleIndex = 0; TriangleIndex < Triangles.Num(); ++TriangleIndex)
+			{
+				const int32 Triangle = Triangles[TriangleIndex];
+				Indices[TriangleIndex * 3] = Triangle * 3;
+				Indices[TriangleIndex * 3 + 1] = Triangle * 3 + 1;
+				Indices[TriangleIndex * 3 + 2] = Triangle * 3 + 2;
+			}
+			return Indices;
+		}
+
+		template<typename T>
+		TArray<T> Gather(const TArray<T>& Array, const TArray<int32>& Indices)
+		{
+			TArray<T> Result;
+			Result.SetNum(Indices.Num());
+			for (int32 Index = 0; Index < Indices.Num(); ++Index)
+			{
+				Result[Index] = Array[Indices[Index]];
+			}
+			return Result;
+		}
+
+		template<typename T1, typename T2>
+		TArray<T2> Map(const TArray<T1>& Array, const TArray<T2>& Mapper)
+		{
+			TArray<int32> Result;
+			Result.SetNum(Array.Num());
+			for (int32 Index = 0; Index < Array.Num(); ++Index)
+			{
+				Result[Index] = Mapper[Array[Index]];
+			}
+			return Result;
+		}
+
+		TArray<int32> InverseMap(const TArray<int32>& Map, const int32 Num)
+		{
+			TArray<int32> InversedMap;
+			InversedMap.SetNum(Num);
+			for (int32 Index = 0; Index < Map.Num(); ++Index)
+			{
+				InversedMap[Map[Index]] = Index;
+			}
+			return InversedMap;
+		}
+
+		TArray<TArray<int32>> GetSectionToTriangleIndices(const FSkeletalMeshLODRenderData& LODData, const TArray<int32>& TriangleIndices)
+		{
+			const TArray<FSkelMeshRenderSection>& Sections = LODData.RenderSections;
+			const int32 NumSections = Sections.Num();
+			TArray<TArray<int32>> SectionToTriangleIndices;
+			SectionToTriangleIndices.SetNumZeroed(NumSections);
+			for (int32 TriangleIndex : TriangleIndices)
+			{
+				for (int32 SectionIndex = 0; SectionIndex < NumSections; ++SectionIndex)
+				{
+					const FSkelMeshRenderSection& Section = Sections[SectionIndex];
+					const int32 Start = Section.BaseIndex;
+					const int32 End = Start + Section.NumTriangles - 1;
+					if (Start <= TriangleIndex && TriangleIndex <= End)
+					{
+						SectionToTriangleIndices[SectionIndex].Add(TriangleIndex);
+						break;
+					}
+				}
+			}
+
+			return SectionToTriangleIndices;
+		}
+
+		template<typename SrcT, typename DstT>
+		TArray<DstT> ConvertArray(const TArray<SrcT>& SrcArray)
+		{
+			TArray<DstT> DstArray;
+			DstArray.SetNumUninitialized(SrcArray.Num());
+			for (int32 Index = 0; Index < SrcArray.Num(); ++Index)
+			{
+				DstArray[Index] = SrcArray[Index];
+			}
+			return DstArray;
+		}
+
+		TArray<uint32> OffsetElements(const TArray<uint32>& Array, int32 Offset)
 		{
-			const FSkeletalMeshRenderData* RenderData = Asset.GetResourceForRendering();
+			TArray<uint32> OffsetArray;
+			OffsetArray.SetNumUninitialized(Array.Num());
+			for (int32 Index = 0; Index < Array.Num(); ++Index)
+			{
+				check((int32)Array[Index] + Offset >= 0);
+				OffsetArray[Index] = Array[Index] + Offset;
+			}
+			return OffsetArray;
+		}
+
+		int32 AddTrackWritersFromSkeletalMesh(UE::GeometryCacheHelpers::FGeometryCacheConstantTopologyWriter& Writer, const USkeletalMesh& SkeletalMesh)
+		{
+			const FSkeletalMeshRenderData* RenderData = SkeletalMesh.GetResourceForRendering();
 			constexpr int32 LODIndex = 0;
 			if (!RenderData || !RenderData->LODRenderData.IsValidIndex(LODIndex))
 			{
-				return;
+				UE_LOG(LogChaosClothGenerator, Error, TEXT("SkeletalMesh has no render data. Failed to create track writers."));
+				return 0;
 			}
 			const FSkeletalMeshLODRenderData& LODData = RenderData->LODRenderData[LODIndex];
+
+			const FMeshDescription* MeshDescription = SkeletalMesh.GetMeshDescription(LODIndex);
+			if (!MeshDescription || MeshDescription->IsEmpty())
+			{
+				UE_LOG(LogChaosClothGenerator, Error, TEXT("SkeletalMesh has no mesh description. Failed to create track writers."));
+				return 0;
+			}
+			FSkeletalMeshConstAttributes MeshAttributes(*MeshDescription);
+			if (!MeshAttributes.HasSourceGeometryParts())
+			{
+				UE_LOG(LogChaosClothGenerator, Error, TEXT("SkeletalMesh has no source geometry parts. Failed to create track writers."));
+				return 0;
+			}
+			const int32 NumParts = MeshAttributes.GetNumSourceGeometryParts();
+			if (NumParts <= 0)
+			{
+				UE_LOG(LogChaosClothGenerator, Error, TEXT("SkeletalMesh has no source geometry parts. Failed to create track writers."));
+				return 0;
+			}
+
+			const FSkeletalMeshModel* const ImportModel = SkeletalMesh.GetImportedModel();
+			if (!ImportModel || !ImportModel->LODModels.IsValidIndex(LODIndex))
+			{
+				UE_LOG(LogChaosClothGenerator, Error, TEXT("SkeletalMesh has no import model. Failed to create track writers."));
+				return 0;
+			}
+			const TArray<int32>& MeshToImportVertexMap = ImportModel->LODModels[LODIndex].MeshToImportVertexMap;
 			const int32 NumVertices = GetNumVertices(LODData);
+			TArray<uint32> Indices;
+			LODData.MultiSizeIndexContainer.GetIndexBuffer(Indices);
+			check(Indices.Num() % 3 == 0);
+
+			const TArray<FVector2f> UVs = GetUV0s(LODData);
+			check(UVs.Num() == NumVertices);
+			const TArray<FColor> Colors = GetColors(LODData, NumVertices);
+			check(Colors.Num() == NumVertices);
+
+			const FSkeletalMeshAttributes::FSourceGeometryPartNameConstRef PartNames = MeshAttributes.GetSourceGeometryPartNames();
+			const TArray<TInterval<int32>> ImportPartIntervals = GetImportPartIntervals(MeshAttributes);
+
+			const TArray<int32> VertexToPart = GetVertexToPart(MeshToImportVertexMap, ImportPartIntervals);
+			const TArray<TArray<int32>> PartToVertices = GetPartToIndices(VertexToPart, NumParts);
+			const TOptional<TArray<int32>> OptionalTriangleToPart = GetTriangleToPart(Indices, VertexToPart);
+			if (!OptionalTriangleToPart.IsSet())
+			{
+				UE_LOG(LogChaosClothGenerator, Error, TEXT("Assuming all vertices in the same triangle are in the same part but some triangles are not. Failed to create track writers."));
+				return 0;
+			}
+			const TArray<int32>& TriangleToPart = OptionalTriangleToPart.GetValue();
+			const TArray<TArray<int32>> PartToTriangles = GetPartToIndices(TriangleToPart, NumParts);
+			for (int32 PartIndex = 0; PartIndex < NumParts; ++PartIndex)
+			{
+				const FName TrackName = PartNames.Get(PartIndex);
+				using FTrackWriter = UE::GeometryCacheHelpers::FGeometryCacheConstantTopologyWriter::FTrackWriter;
+				FTrackWriter& TrackWriter = Writer.AddTrackWriter(TrackName);
+				const TArray<int32>& VertexIndices = PartToVertices[PartIndex];
+				const TArray<int32> InverseVertexMap = InverseMap(VertexIndices, NumVertices);
+				const TArray<int32>& TriangleIndices = PartToTriangles[PartIndex];
+				const int32 StartImportedVertex = ImportPartIntervals[PartIndex].Min;
+				TrackWriter.UVs = Gather(UVs, VertexIndices);
+				TrackWriter.Colors = Gather(Colors, VertexIndices);
+				TrackWriter.ImportedVertexNumbers = OffsetElements(ConvertArray<int32, uint32>(Gather(MeshToImportVertexMap, VertexIndices)), -StartImportedVertex);
+				TrackWriter.SourceVertexIndices = VertexIndices;
+
+				const TArray<TArray<int32>> SectionToTriangleIndices = GetSectionToTriangleIndices(LODData, TriangleIndices);
+				TArray<uint32> PartIndices;
+				PartIndices.Reserve(TriangleIndices.Num() * 3);
+				for (int32 SectionIndex = 0; SectionIndex < SectionToTriangleIndices.Num(); ++SectionIndex)
+				{
+					const TArray<int32>& SectionTriangleIndices = SectionToTriangleIndices[SectionIndex];
+					if (SectionTriangleIndices.IsEmpty())
+					{
+						continue;
+					}
+					TArray<uint32> SectionIndices = ConvertArray<int32, uint32>(Map(Gather(Indices, GetIndicesFromTriangles(SectionTriangleIndices)), InverseVertexMap));
+					check(SectionIndices.Num() % 3 == 0);
+					FGeometryCacheMeshBatchInfo BatchInfo;
+					BatchInfo.StartIndex = PartIndices.Num();
+					BatchInfo.NumTriangles = SectionIndices.Num() / 3;
+					BatchInfo.MaterialIndex = LODData.RenderSections[SectionIndex].MaterialIndex;
+					TrackWriter.BatchesInfo.Add(BatchInfo);
+					PartIndices.Append(MoveTemp(SectionIndices));
+				}
+				TrackWriter.Indices = PartIndices;
+			}
+			Writer.AddMaterials(GetMaterialInterfaces(SkeletalMesh));
+			return NumParts;
+		}
+
+		TArray<TArray<FVector3f>> Gather(const TArrayView<TArray<FVector3f>>& Positions, const TArray<int32>& Indices)
+		{
+			if (Indices.IsEmpty())
+			{
+				return TArray<TArray<FVector3f>>();
+			}
+			const int32 NumFrames = Positions.Num();
+			TArray<TArray<FVector3f>> Result;
+			Result.SetNum(NumFrames);
+			for (int32 Frame = 0; Frame < NumFrames; ++Frame)
+			{
+				const TArray<FVector3f>& FramePositions = Positions[Frame];
+				Result[Frame] = Gather(FramePositions, Indices);
+			}
+			return Result;
+		}
+		
+		void SaveGeometryCache(UGeometryCache& GeometryCache, const USkeletalMesh& SkeletalMesh, TArrayView<TArray<FVector3f>> PositionsToMoveFrom)
+		{
+			constexpr int32 LODIndex = 0;
+			const int32 NumVertices = GetNumVertices(SkeletalMesh);
 			PositionsToMoveFrom = ShrinkToValidFrames(PositionsToMoveFrom, NumVertices);
+			if (PositionsToMoveFrom.IsEmpty())
+			{
+				UE_LOG(LogChaosClothGenerator, Error, TEXT("There is no valid data in PositionsToMoveFrom. No geometry cache is saved."));
+				return;
+			}
 		
 			using UE::GeometryCacheHelpers::FGeometryCacheConstantTopologyWriter;
 			using UE::GeometryCacheHelpers::AddTrackWriterFromSkinnedAsset;
 			using FTrackWriter = FGeometryCacheConstantTopologyWriter::FTrackWriter;
 			FGeometryCacheConstantTopologyWriter Writer(GeometryCache);
-			const int32 Index = AddTrackWriterFromSkinnedAsset(Writer, Asset);
-			if (Index == INDEX_NONE)
+
+			const int32 NumTracks = AddTrackWritersFromSkeletalMesh(Writer, SkeletalMesh);
+			if (NumTracks == 0)
 			{
+				UE_LOG(LogChaosClothGenerator, Error, TEXT("Failed to add track writers. No geometry cache is saved."));
 				return;
 			}
-			FTrackWriter& TrackWriter = Writer.GetTrackWriter(Index);
-			TrackWriter.ImportedVertexNumbers = ImportedVertexNumbers;
-			TrackWriter.WriteAndClose(PositionsToMoveFrom);
+			for (int32 Index = 0; Index < NumTracks; ++Index)
+			{
+				FTrackWriter& TrackWriter = Writer.GetTrackWriter(Index);
+				check(TrackWriter.SourceVertexIndices.IsSet());
+				TrackWriter.WriteAndClose(Gather(PositionsToMoveFrom, TrackWriter.SourceVertexIndices.GetValue()));
+			}
 		}
 		
 		class FTimeScope
@@ -187,37 +552,42 @@ namespace UE::Chaos::ClothGenerator
 			constexpr bool bPromptToSave = false;
 			FEditorFileUtils::PromptForCheckoutAndSave(PackagesToSave, bCheckDirty, bPromptToSave);
 		}
-		
-		TOptional<TArray<int32>> GetMeshImportVertexMap(const USkinnedAsset& SkeletalMeshAsset, const UChaosClothAsset& ClothAsset)
+
+		bool AreAssetsConsistent(const USkinnedAsset& SkeletalMeshAsset, const UChaosClothAsset& ClothAsset)
 		{
+			if (!SkeletalMeshAsset.IsA<USkeletalMesh>())
+			{
+				UE_LOG(LogChaosClothGenerator, Warning, TEXT("SkeletalMeshAsset is not a SkeletalMesh."));
+				return false;
+			}
 			constexpr int32 LODIndex = 0;
 			const TOptional<TArray<int32>> None;
 			const FSkeletalMeshModel* const MLDModel = SkeletalMeshAsset.GetImportedModel();
 			if (!MLDModel || !MLDModel->LODModels.IsValidIndex(LODIndex))
 			{
-				return None;
+				return false;
 			}
 			const FSkeletalMeshLODModel& MLDLOD = MLDModel->LODModels[LODIndex];
 			const TArray<int32>& Map = MLDLOD.MeshToImportVertexMap;
 			if (Map.IsEmpty())
 			{
 				UE_LOG(LogChaosClothGenerator, Warning, TEXT("MeshToImportVertexMap is empty. MLDeformer Asset should be an imported SkeletalMesh (e.g. from fbx)."));
-				return None;
+				return false;
 			}
 			const FSkeletalMeshModel* const ClothModel = ClothAsset.GetImportedModel();
 			if (!ClothModel || !ClothModel->LODModels.IsValidIndex(LODIndex))
 			{
 				UE_LOG(LogChaosClothGenerator, Warning, TEXT("ClothAsset has no imported model."));
-				return None;
+				return false;
 			}
 			const FSkeletalMeshLODModel& ClothLOD = ClothModel->LODModels[LODIndex];
-		
+
 			if (MLDLOD.NumVertices != ClothLOD.NumVertices || MLDLOD.Sections.Num() != ClothLOD.Sections.Num())
 			{
 				UE_LOG(LogChaosClothGenerator, Warning, TEXT("SkeletalMeshAsset and ClothAsset have different number of vertices or sections. Check if the assets have the same mesh."));
-				return None;
+				return false;
 			}
-			
+
 			for (int32 SectionIndex = 0; SectionIndex < MLDLOD.Sections.Num(); ++SectionIndex)
 			{
 				const FSkelMeshSection& MLDSection = MLDLOD.Sections[SectionIndex];
@@ -225,7 +595,7 @@ namespace UE::Chaos::ClothGenerator
 				if (MLDSection.NumVertices != ClothSection.NumVertices)
 				{
 					UE_LOG(LogChaosClothGenerator, Warning, TEXT("SkeletalMeshAsset and ClothAsset have different number of vertices in section %d. Check if the assets have the same mesh."), SectionIndex);
-					return None;
+					return false;
 				}
 				for (int32 VertexIndex = 0; VertexIndex < MLDSection.NumVertices; ++VertexIndex)
 				{
@@ -234,12 +604,12 @@ namespace UE::Chaos::ClothGenerator
 					if (!MLDPosition.Equals(ClothPosition, UE_KINDA_SMALL_NUMBER))
 					{
 						UE_LOG(LogChaosClothGenerator, Warning, TEXT("SkeletalMeshAsset and ClothAsset have different vertex positions. Check if the assets have the same vertex order."));
-						return None;
+						return false;
 					}
 				}
 			}
-		
-			return Map;
+
+			return true;
 		}
 	};
 
@@ -318,7 +688,6 @@ namespace UE::Chaos::ClothGenerator
 
 		TArray<int32> FramesToSimulate;
 		TArray<TArray<FVector3f>> SimulatedPositions;
-		TArray<uint32> ImportedVertexNumbers;
 		UGeometryCache* Cache = nullptr;
 
 		std::atomic<int32> NumSimulatedFrames = 0;
@@ -679,9 +1048,8 @@ namespace UE::Chaos::ClothGenerator
 			return;
 		}
 
-		using UE::Chaos::ClothGenerator::Private::GetMeshImportVertexMap;
-		TOptional<TArray<int32>> OptionalMap = GetMeshImportVertexMap(*Properties->SkeletalMeshAsset, *Properties->ClothAsset);
-		if (!OptionalMap)
+		using UE::Chaos::ClothGenerator::Private::AreAssetsConsistent;
+		if (!AreAssetsConsistent(*Properties->SkeletalMeshAsset, *Properties->ClothAsset))
 		{
 			PendingAction = EClothGeneratorActions::NoAction;
 			return;
@@ -724,9 +1092,6 @@ namespace UE::Chaos::ClothGenerator
 		TaskResource->StartTime = FDateTime::UtcNow();
 		TaskResource->LastUpdateTime = TaskResource->StartTime;
 	
-		const TArray<int32>& Map = OptionalMap.GetValue();
-		TaskResource->ImportedVertexNumbers = TArray<uint32>(reinterpret_cast<const uint32*>(Map.GetData()), Map.Num());
-	
 		PendingAction = EClothGeneratorActions::TickGenerate;
 	}
 	
@@ -801,7 +1166,9 @@ namespace UE::Chaos::ClothGenerator
 	
 			using UE::Chaos::ClothGenerator::Private::SaveGeometryCache;
 			using UE::Chaos::ClothGenerator::Private::SavePackage;
-			SaveGeometryCache(*TaskResource->Cache, *Properties->ClothAsset, TaskResource->ImportedVertexNumbers, TaskResource->SimulatedPositions);
+			const USkeletalMesh* const SkeletalMesh = Cast<const USkeletalMesh>(Properties->SkeletalMeshAsset);
+			check(SkeletalMesh);
+			SaveGeometryCache(*TaskResource->Cache, *SkeletalMesh, TaskResource->SimulatedPositions);
 			SavePackage(*TaskResource->Cache);
 		}
 		if (bCancelled)
diff --git a/Engine/Plugins/Runtime/GeometryCache/Source/GeometryCache/Public/GeometryCacheConstantTopologyWriter.h b/Engine/Plugins/Runtime/GeometryCache/Source/GeometryCache/Public/GeometryCacheConstantTopologyWriter.h
--- a/Engine/Plugins/Runtime/GeometryCache/Source/GeometryCache/Public/GeometryCacheConstantTopologyWriter.h
+++ b/Engine/Plugins/Runtime/GeometryCache/Source/GeometryCache/Public/GeometryCacheConstantTopologyWriter.h
@@ -71,6 +71,7 @@ namespace UE::GeometryCacheHelpers
 			TArray<FColor> Colors;
 			TArray<uint32> ImportedVertexNumbers;
 			TArray<FGeometryCacheMeshBatchInfo> BatchesInfo;
+			TOptional<TArray<int32>> SourceVertexIndices;
 
 			/**
 			 * Move the position data to the cache track and close the TrackWriter. 
@@ -109,6 +110,7 @@ namespace UE::GeometryCacheHelpers
 
 	/**
 	 * @brief This will create a track writer and fill in the track writer's data (indices, UVs, materials .etc) from the skinned asset.
+	 * This only creates one track for the skinned asset. If you want to create multiple tracks based on imported mesh info, checkout AddTrackWriterFromSkinnedAsset() in ChaosClothGenerator.cpp. 
 	 * Usage:
 	 * 	FGeometryCacheConstantTopologyWriter Writer(MyCache);
 	 * 	int32 Index = AddTrackWriterFromSkinnedAsset(Writer, Asset);
