Fix unity buld error

diff --git a/Engine/Plugins/Animation/MLDeformer/ChaosClothGenerator/Source/ChaosClothGenerator/Private/ChaosClothGeneratorModule.cpp b/Engine/Plugins/Animation/MLDeformer/ChaosClothGenerator/Source/ChaosClothGenerator/Private/ChaosClothGeneratorModule.cpp
--- a/Engine/Plugins/Animation/MLDeformer/ChaosClothGenerator/Source/ChaosClothGenerator/Private/ChaosClothGeneratorModule.cpp
+++ b/Engine/Plugins/Animation/MLDeformer/ChaosClothGenerator/Source/ChaosClothGenerator/Private/ChaosClothGeneratorModule.cpp
@@ -1,6 +1,7 @@
 // Copyright Epic Games, Inc. All Rights Reserved.
 #include "ChaosClothGeneratorToolsMenuExtender.h"
 #include "MLDeformerEditorToolkit.h"
+#include "Modules/ModuleManager.h"
 #include "PropertyEditorModule.h"
 #include "SClothGeneratorWidget.h"
 
@@ -24,4 +25,4 @@ namespace UE::Chaos::ClothGenerator
 	};
 };
 
-IMPLEMENT_MODULE(UE::Chaos::ClothGenerator::FChaosClothGeneratorModule, ChaosClothGenerator)
\ No newline at end of file
+IMPLEMENT_MODULE(UE::Chaos::ClothGenerator::FChaosClothGeneratorModule, ChaosClothGenerator)
