Fix crash in Cloth Generator when there is no skeleton assigned to a cloth asset.

diff --git a/Engine/Plugins/Animation/MLDeformer/ChaosClothGenerator/Source/ChaosClothGenerator/Private/ChaosClothGenerator.cpp b/Engine/Plugins/Animation/MLDeformer/ChaosClothGenerator/Source/ChaosClothGenerator/Private/ChaosClothGenerator.cpp
--- a/Engine/Plugins/Animation/MLDeformer/ChaosClothGenerator/Source/ChaosClothGenerator/Private/ChaosClothGenerator.cpp
+++ b/Engine/Plugins/Animation/MLDeformer/ChaosClothGenerator/Source/ChaosClothGenerator/Private/ChaosClothGenerator.cpp
@@ -938,7 +938,7 @@ namespace UE::Chaos::ClothGenerator
 		USkeleton* const Skeleton = ClothAsset ? ClothAsset->GetSkeleton() : nullptr;
 		const int32 NumBones = ReferenceSkeleton ? ReferenceSkeleton->GetNum() : 0;
 
-		if (!ReferenceSkeleton || NumBones == 0)
+		if (!ReferenceSkeleton || !Skeleton || NumBones == 0)
 		{
 			return TArray<FTransform>();
 		}
